Copying //UE4/Dev-Rendering to //UE4/Dev-Main (Source: //UE4/Dev-Rendering @ 3072736)

==========================
MAJOR FEATURES + CHANGES
==========================

Change 3055495 on 2016/07/19 by [REDACTED_PERSON]

	Allow Noise material node on mobile

	No reason to exclude mobile, except for Fast Gradient Noise, which uses 3D textures. Allow this node on ES2 for all of the other noise functions.

Change 3055602 on 2016/07/19 by [REDACTED_PERSON]

	Fix crash bug in D3D11 RHI when selecting adapters.
	 - Array of adapter descriptors will get out of sync with the adapter index if any adapter is skipped (e.g. the Microsoft Basic Render Device).

Change 3055890 on 2016/07/19 by [REDACTED_PERSON]

	Improved the assert in LoadModuleChecked so we won't have to check the log to see which module it was

Change 3055891 on 2016/07/19 by [REDACTED_PERSON]

	Fixed Global Distance Field not dirtying previous object position on UpdateTransform - left behind a phantom shadow on teleports
	* This will effectively double partial distiance field update costs until clipping of the update regions is implemented

Change 3055892 on 2016/07/19 by [REDACTED_PERSON]

	Higher poly light source shapes drawn into reflection captures

Change 3055893 on 2016/07/19 by [REDACTED_PERSON]

	More info to 'Incompatible surface format' GNM assert

Change 3055904 on 2016/07/19 by [REDACTED_PERSON]

	Reflection environment normalization improvements
	* Indirect specular from reflection captures is now mixed with indirect diffuse from lightmaps based on roughness, such that a mirror surface will have no mixing.  Reflection captures now match other reflection methods like SSR and planar reflections much more closely.
	* When a stationary skylight is present, Reflection captures are now normalized as if the initial skylight will always be present, giving consistent results with static skylight reflections.  The skylight and reflection captures with sky removed used to be normalized separately, compacting the relative brightness between the sky and scene.
	* Added r.ReflectionEnvironmentLightmapMixing for debugging lightmap mixing issues.  This toggle was previously not possible due to prenormalizing the capture data.
	* The standard deferred reflection path (r.DoTiledReflections 0) can no longer match the results of the compute path or base pass reflections, as it would require MRT to accumulate the average brightness
	* Removed unused r.DiffuseFromCaptures
	* Cost of reflection environment on PS4 increased from 1.52ms -> 1.75ms with this change, but decreased back to 1.58ms by reducing tile size to 8x8

Change 3055905 on 2016/07/19 by [REDACTED_PERSON]

	Workaround for RTDF shadows not working on PS4 - manual clear of ObjectIndirectArguments instead of RHICmdList.ClearUAV

Change 3059486 on 2016/07/21 by [REDACTED_PERSON]

	Testing #uecritical

Change 3060558 on 2016/07/21 by [REDACTED_PERSON]

	Fixed skylight with specified cubemap being black

Change 3061999 on 2016/07/22 by [REDACTED_PERSON]

	Disable old AMD driver hacks for DX11.  QA has already tested with them off and given thumbs up.

Change 3062241 on 2016/07/22 by [REDACTED_PERSON]

	Fixed bug in RHISupportsSeparateMSAAAndResolveTextures that was preventing MSAA for any non-Vulkan platforms

Change 3062244 on 2016/07/22 by [REDACTED_PERSON]

	Discard old prenormalized reflection environment data on load

Change 3062283 on 2016/07/22 by [REDACTED_PERSON]

	MSAA support for the forward renderer
	* AntiAliasing method is chosen in Rendering project settings, DefaultSettings category
	* Deferred passes like shadow projection, fogging and decals are only computed per-pixel and can introduce aliasing
	* Added Rendering project setting VertexFoggingForOpaque, which makes height fog cheaper and work properly with MSAA
	* The AntiAliasing method in PostProcessSettings has been removed, this may affect existing content
	* Added r.MSAACount which defaults to 4
	* Integrated wide custom resolve filter from Oculus renderer, controlled by r.WideCustomResolve
	* GBuffer targets are no longer allocated when using the forward renderer
	* Decal blend modes that write to the GBuffer fall back to SceneColor emissive only

Change 3062666 on 2016/07/23 by [REDACTED_PERSON]

	Added legend to streaming accuracy viewmodes
	Added a new helper class FRenderTargetTemp to be reused in different canvas rendering.
	Exposed the pass through pixel shader so that it can be reused.

Change 3063023 on 2016/07/25 by [REDACTED_PERSON]

	Fix "RecompileShaders Changed" when using Cook On The Fly.

Change 3063078 on 2016/07/25 by [REDACTED_PERSON]

	Add -emitdrawevents command line option to emit draw events by default. This is useful when capturing with Renderdoc

Change 3063315 on 2016/07/25 by [REDACTED_PERSON]

	Fix div 0 in motion blur. This caused artifacts in some fairly common cases

Change 3063897 on 2016/07/25 by [REDACTED_PERSON]

	Fixed missing qualifier on interpolants

Change 3064559 on 2016/07/26 by [REDACTED_PERSON]

	Fix for cooker crash with BC6H textures (XB1, but may affect other platforms). Also fixes corruption issue with texture slices not being a multiple of 4 pixels (expanding as necessary), courtesy of [REDACTED_PERSON] at the Coalition
	Tested fix on xbox, PC and PS4, using QAGame

Change 3064896 on 2016/07/26 by [REDACTED_PERSON]

	Fix compile errors on PS4 (the variable "sample" was conflicting with a keyword, causing compile errors). Also making encoding consistent on new shaders (ansi rather than UTF16)

Change 3064913 on 2016/07/26 by [REDACTED_PERSON]

	Fix spelling of "Editor, Tools, Monolithics & DDC" node in Dev-Rendering build settings.

Change 3065326 on 2016/07/26 by [REDACTED_PERSON]

	Fixed UnbuiltInstanceBoundsList not being reset correctly, creating broken rendered primitives.

Change 3065541 on 2016/07/26 by [REDACTED_PERSON]

	Materials with a GBuffer SceneTexture lookup will fail to compile with forward shading

Change 3065543 on 2016/07/26 by [REDACTED_PERSON]

	Restored DetailMode changes causing a FGlobalComponentRecreateRenderStateContext - accidental removal from cl 2969413

Change 3065545 on 2016/07/26 by [REDACTED_PERSON]

	Added material property bNormalCurvatureToRoughness, which can slightly reduce aliasing.  Tweakable impact with r.NormalCurvatureToRoughnessScale.
	Fixed reflection capture feedback with base pass reflections

Change 3066783 on 2016/07/27 by [REDACTED_PERSON]

	Moved PreShadowCacheDepthZ out of FSceneRenderTargets and into FScene, which fixes issues with cached preshadows and multiple scenes, including HighResScreenShot
	Disabled GMinScreenRadiusForShadowCaster on per-object shadows, which fixes popping when trying to increase shadow resolution from the defaults (r.Shadow.TexelsPerPixel 3)

Change 3066794 on 2016/07/27 by [REDACTED_PERSON]

	Fixed crash rendering planar reflections due to NULL PostProcessSettings

Change 3067412 on 2016/07/27 by [REDACTED_PERSON]

	Fix for OpenGL4 with uint interpolator

Change 3068470 on 2016/07/28 by [REDACTED_PERSON]

	Fixed crash rendering translucency with translucent shadows which were determined to be invisible

Change 3069046 on 2016/07/28 by [REDACTED_PERSON]

	Handle null Family in SetupAntiAliasingMethod

Change 3069059 on 2016/07/28 by [REDACTED_PERSON]

	Added r.ReflectionEnvironmentBeginMixingRoughness (.1) and r.ReflectionEnvironmentEndMixingRoughness (.3), which can be used to tweak the lightmap mixing heuristc, or revert to previous behavior (mixing even on a mirror surface)

Change 3069391 on 2016/07/28 by [REDACTED_PERSON]

	Fixed AverageBrightness being applied to reflections in gamma space in the mobile base pass, causing ES2 reflections to be overbright

Change 3070369 on 2016/07/29 by [REDACTED_PERSON]

	r.ReflectionEnvironmentBeginMixingRoughness and r.ReflectionEnvironmentEndMixingRoughness set to 0 can be used to achieve old non-roughness based lightmap mixing

Change 3070370 on 2016/07/29 by [REDACTED_PERSON]

	Bumped reflection capture DDC version to get rid of legacy prenormalized data

Change 3070680 on 2016/07/29 by [REDACTED_PERSON]

	Fix slate ensure that is most likely a timing issue exposed by rendering.
	#ue-33902

Change 3070811 on 2016/07/29 by [REDACTED_PERSON]

	Fix ProjectLauncher errors when loading old versions
	#ue-33939

Change 3070971 on 2016/07/29 by [REDACTED_PERSON]

	Updated ListTextures outputs to fix cooked VS non cooked differences and also to put enphasis on disk VS memory

Change 3071452 on 2016/07/31 by [REDACTED_PERSON]

	Updated the legend description for the (texture streaming) primitive distance accuracy view mode

diff --git a/Engine/Config/BaseEngine.ini b/Engine/Config/BaseEngine.ini
--- a/Engine/Config/BaseEngine.ini
+++ b/Engine/Config/BaseEngine.ini
@@ -158,7 +158,7 @@ DefaultSoundName=/Engine/EngineSounds/WhiteNoise.WhiteNoise
 MaximumLoopIterationCount = 1000000
 ; Allows Blueprint classes based on a native Actor or Component subclass to tick even if their parent has bCanEverTick flag set to false.
 bCanBlueprintsTickByDefault=true
-; Controls whether anim blueprint nodes that access member variables of their class directly should use the optimized path that avoids a thunk to the Blueprint VM
+; Controls whether anim blueprint nodes that access member variables of their class directly should use the optimized path that avoids a thunk to the Blueprint VM.
 bOptimizeAnimBlueprintMemberVariableAccess=true
 CameraRotationThreshold=45.0
 CameraTranslationThreshold=10000
diff --git a/Engine/Shaders/BasePassCommon.usf b/Engine/Shaders/BasePassCommon.usf
--- a/Engine/Shaders/BasePassCommon.usf
+++ b/Engine/Shaders/BasePassCommon.usf
@@ -10,7 +10,7 @@
 // Translucent materials need to compute fogging in the forward shading pass
 // Materials that read from scene color skip getting fogged, because the contents of the scene color lookup have already been fogged
 // This is not foolproof, as any additional color the material adds will then not be fogged correctly
-#define NEEDS_BASEPASS_FOGGING			(ENABLE_TRANSLUCENCY_VERTEX_FOG && (MATERIALBLENDING_ANY_TRANSLUCENT && !MATERIAL_USES_SCENE_COLOR_COPY))
+#define NEEDS_BASEPASS_FOGGING			(ENABLE_TRANSLUCENCY_VERTEX_FOG && MATERIALBLENDING_ANY_TRANSLUCENT && !MATERIAL_USES_SCENE_COLOR_COPY || !MATERIALBLENDING_ANY_TRANSLUCENT && VERTEX_FOGGING_FOR_OPAQUE)
 #define NEEDS_LIGHTMAP					(NEEDS_LIGHTMAP_COORDINATE)
 
 #define USES_GBUFFER						(FEATURE_LEVEL >= FEATURE_LEVEL_SM4 && (MATERIALBLENDING_SOLID || MATERIALBLENDING_MASKED) && !SIMPLE_FORWARD_SHADING && !FORWARD_SHADING)
diff --git a/Engine/Shaders/BasePassPixelShader.usf b/Engine/Shaders/BasePassPixelShader.usf
--- a/Engine/Shaders/BasePassPixelShader.usf
+++ b/Engine/Shaders/BasePassPixelShader.usf
@@ -21,6 +21,15 @@
 #define PREV_FRAME_COLOR	1
 #include "ScreenSpaceRayCast.usf"
 
+float NormalCurvatureToRoughness(float3 WorldNormal)
+{
+    float3 dNdx = ddx(WorldNormal);
+    float3 dNdy = ddy(WorldNormal);
+    float x = dot(dNdx, dNdx);
+    float y = dot(dNdy, dNdy);
+    return max(x, y);
+}
+
 #if TRANSLUCENT_SELF_SHADOWING
 
 	#include "ShadowProjectionCommon.usf"
@@ -78,7 +87,7 @@ float4 ReflectionPositionAndRadius;
 float ReflectionShape;
 float4x4 BoxTransform;
 float4 BoxScales;
-float3 CaptureOffset;
+float4 CaptureOffsetAndAverageBrightness;
 
 uint MortonCode( uint x )
 {
@@ -97,10 +106,12 @@ half3 GetImageBasedReflectionLighting(FMaterialPixelParameters MaterialParameter
 	float3 RayDirection = 2 * dot( V, N ) * N - V;
 	half NoV = saturate(dot(N, V));
 
-	float4 SpecularIBL = float4(0,0,0,0.0);
+	float4 SpecularIBL = float4(0, 0, 0, 0);
 	float3 CaptureVector = MaterialParameters.AbsoluteWorldPosition - ReflectionPositionAndRadius.xyz;
 	float CaptureVectorLength = length(CaptureVector);
-	
+
+	float2 CompositedAverageBrightness = float2(0.0f, 1.0f);
+
 	BRANCH
 	if (CaptureVectorLength < ReflectionPositionAndRadius.w)
 	{
@@ -110,11 +121,11 @@ half3 GetImageBasedReflectionLighting(FMaterialPixelParameters MaterialParameter
 		float3 ProjectedCaptureVector;
 		if (ReflectionShape > 0.0f)
 		{
-			ProjectedCaptureVector = GetLookupVectorForBoxCapture(RayDirection, MaterialParameters.AbsoluteWorldPosition, ReflectionPositionAndRadius, BoxTransform, BoxScales, CaptureOffset, DistanceAlpha);
+			ProjectedCaptureVector = GetLookupVectorForBoxCapture(RayDirection, MaterialParameters.AbsoluteWorldPosition, ReflectionPositionAndRadius, BoxTransform, BoxScales, CaptureOffsetAndAverageBrightness.xyz, DistanceAlpha);
 		}
 		else
 		{
-			ProjectedCaptureVector = GetLookupVectorForSphereCapture(RayDirection, MaterialParameters.AbsoluteWorldPosition, ReflectionPositionAndRadius, NormalizedDistanceToCapture, CaptureOffset, DistanceAlpha);
+			ProjectedCaptureVector = GetLookupVectorForSphereCapture(RayDirection, MaterialParameters.AbsoluteWorldPosition, ReflectionPositionAndRadius, NormalizedDistanceToCapture, CaptureOffsetAndAverageBrightness.xyz, DistanceAlpha);
 		}
 			
 		half AbsoluteSpecularMip = ComputeReflectionCaptureMipFromRoughness(Roughness, View.ReflectionCubemapMaxMip);	
@@ -126,36 +137,50 @@ half3 GetImageBasedReflectionLighting(FMaterialPixelParameters MaterialParameter
 #endif
 		Sample *= DistanceAlpha;
 		SpecularIBL = Sample;
+
+		float AverageBrightness = CaptureOffsetAndAverageBrightness.w;
+		CompositedAverageBrightness.x += AverageBrightness * DistanceAlpha * CompositedAverageBrightness.y;
+		CompositedAverageBrightness.y *= 1 - DistanceAlpha;
 	}	
 
-#if ALLOW_STATIC_LIGHTING
-	FLATTEN
-	if (View.UseLightmaps > 0)
-	{
-		// Note: make sure this matches the lightmap mixing done on opaque (ReflectionEnvironmentTiledDeferredMain)
-		SpecularIBL.rgb *= IndirectIrradiance;
-	}
-#endif
+	float3 ExtraIndirectSpecular = 0;
 
 #if ENABLE_SKY_LIGHT
 	BRANCH
 	if (SkyLightParameters.y > 0 && SpecularIBL.a < .999f)
 	{
+		float SkyAverageBrightness = 1.0f;
+		float3 SkyLighting = GetSkyLightReflection(RayDirection, Roughness, SkyAverageBrightness);
+
 		// Normalize for static skylight types which mix with lightmaps
-		bool bNormalize = SkyLightParameters.z < 1 && View.UseLightmaps;
-		float3 SkyLighting = GetSkyLightReflection(RayDirection, Roughness, bNormalize);
+		bool bNormalize = SkyLightParameters.z < 1 && ALLOW_STATIC_LIGHTING;
 
 		FLATTEN
 		if (bNormalize)
 		{
-			SkyLighting *= IndirectIrradiance;
+			// Add in sky wherever reflection captures don't have coverage
+			SpecularIBL.rgb += (1 - SpecularIBL.a) * SkyLighting;
+			CompositedAverageBrightness.x += SkyAverageBrightness * CompositedAverageBrightness.y;
+		}
+		else
+		{
+			ExtraIndirectSpecular += SkyLighting;
 		}
-
-		// Add in sky wherever reflection captures don't have coverage
-		SpecularIBL.rgb += (1 - SpecularIBL.a) * SkyLighting;
 	}
 #endif
 
+#if ALLOW_STATIC_LIGHTING
+	// Note: make sure this matches the lightmap mixing done on opaque (ReflectionEnvironmentTiledDeferredMain)
+	SpecularIBL.rgb *= ComputeMixingWeight(IndirectIrradiance, CompositedAverageBrightness.x, Roughness);
+#endif
+
+	SpecularIBL.rgb += (1 - SpecularIBL.a) * ExtraIndirectSpecular;
+
+	// Factors derived from EnvBRDFApprox( SpecularColor, 1, 1 ) == SpecularColor * 0.4524 - 0.0024
+	float3 SpecularBounce = 0.45f * SpecularColor * IndirectIrradiance;
+	// Replace reflection captures with indirect diffuse when we're rendering to a reflection capture, to avoid a feedback loop
+	SpecularIBL.rgb = lerp(SpecularIBL.rgb, SpecularBounce, View.RenderingReflectionCaptureMask);
+
 #if MATERIAL_SSR && !FORWARD_SHADING
 	if( View.CameraCut == 0 )
 	{
@@ -549,21 +574,8 @@ void GetPrecomputedIndirectLightingAndSkyLight(
 
 	OutSubsurfaceLighting += SkySubsurfaceLighting;
 
-	BRANCH
-	if( View.UseLightmaps > 0 || !ALLOW_STATIC_LIGHTING)
-	{
-		// Sky lighting must contribute to IndirectIrradiance for ReflectionEnvironment lightmap mixing
-		OutDiffuseLighting += SkyDiffuseLighting;
-	}
-	else
-	{
-	#if (TRANSLUCENCY_LIGHTING_SURFACE || TRANSLUCENCY_LIGHTING_SURFACE_PERPIXEL) && FEATURE_LEVEL >= FEATURE_LEVEL_SM5
-		float DiffuseMip = ComputeReflectionCaptureMipFromRoughness(1, View.ReflectionCubemapMaxMip);
-		float4 DiffuseIBL = TextureCubeArraySampleLevel(ReflectionCubemap, ReflectionCubemapSampler, DiffuseDir, CubemapArrayIndex, DiffuseMip);
-		
-		OutDiffuseLighting += DiffuseIBL.rgb + (1 - DiffuseIBL.a) * SkyDiffuseLighting;
-	#endif
-	}
+	// Sky lighting must contribute to IndirectIrradiance for ReflectionEnvironment lightmap mixing
+	OutDiffuseLighting += SkyDiffuseLighting;
 
 	#if HQ_TEXTURE_LIGHTMAP || LQ_TEXTURE_LIGHTMAP || CACHED_VOLUME_INDIRECT_LIGHTING || CACHED_POINT_INDIRECT_LIGHTING
 		OutIndirectIrradiance = Luminance(OutDiffuseLighting);
@@ -765,6 +777,11 @@ void FPixelShaderInOut_MainPS(
 	float MaterialAO = GetMaterialAmbientOcclusion(PixelMaterialInputs);
 	float Roughness = GetMaterialRoughness(PixelMaterialInputs);
 
+#if MATERIAL_NORMAL_CURVATURE_TO_ROUGHNESS && FORWARD_SHADING
+    float GeometricAARoughness = saturate(NormalCurvatureToRoughness(MaterialParameters.WorldNormal) * View.NormalCurvatureToRoughnessScaleBias.x + View.NormalCurvatureToRoughnessScaleBias.y);
+    Roughness = max(Roughness, GeometricAARoughness);
+#endif
+
 	// 0..1, SubsurfaceProfileId = int(x * 255)
 	float SubsurfaceProfile = 0;
 
@@ -1090,7 +1107,7 @@ void FPixelShaderInOut_MainPS(
 		{
 			FLightAccumulator LightAccumulator = (FLightAccumulator)0;
 
-			LightAccumulator_Add(LightAccumulator, Color, 0, 1.0f);
+			LightAccumulator_Add(LightAccumulator, Color * VertexFog.a + VertexFog.rgb, 0, 1.0f);
 			Out.MRT[0] = RETURN_COLOR(LightAccumulator_GetResult(LightAccumulator));
 
 			#if SUBSURFACE_CHANNEL_MODE == 0
diff --git a/Engine/Shaders/CopyShadowMaps.usf b/Engine/Shaders/CopyShadowMaps.usf
--- a/Engine/Shaders/CopyShadowMaps.usf
+++ b/Engine/Shaders/CopyShadowMaps.usf
@@ -22,10 +22,10 @@ void Copy2DDepthPS(
 struct FCubeCopyGSToPS
 {
 	FScreenVertexOutput Vertex;
-	uint FaceIndex : TEXCOORD1;
+	nointerpolation uint FaceIndex : TEXCOORD1;
 
 	/** Controls which of the cube map faces to rasterize the primitive to, only the value from the first vertex is used. */
-	uint RTIndex : SV_RenderTargetArrayIndex;
+	nointerpolation uint RTIndex : SV_RenderTargetArrayIndex;
 };
  
 /** Allocate space for cloning to all 6 faces. */
diff --git a/Engine/Shaders/MaterialTexCoordScalesPixelShader.usf b/Engine/Shaders/MaterialTexCoordScalesPixelShader.usf
--- a/Engine/Shaders/MaterialTexCoordScalesPixelShader.usf
+++ b/Engine/Shaders/MaterialTexCoordScalesPixelShader.usf
@@ -98,7 +98,7 @@ void Main(
 #endif
 
 	// This default value will make it dark grey when nothing updates the texSample.
-	float3 Result = float3(.02, .02, .02);
+	float3 Result = float3(UNDEFINED_ACCURACY, UNDEFINED_ACCURACY, UNDEFINED_ACCURACY);
 
 	FTexCoordScalesParams Params;
 
diff --git a/Engine/Shaders/MeshTexCoordSizeAccuracyPixelShader.usf b/Engine/Shaders/MeshTexCoordSizeAccuracyPixelShader.usf
--- a/Engine/Shaders/MeshTexCoordSizeAccuracyPixelShader.usf
+++ b/Engine/Shaders/MeshTexCoordSizeAccuracyPixelShader.usf
@@ -25,7 +25,7 @@ void Main(
 	ResolvedView = ResolveView();
 #endif
 
-	float3 Result = float3(.1, .1, .1);
+	float3 Result = float3(UNDEFINED_ACCURACY, UNDEFINED_ACCURACY, UNDEFINED_ACCURACY);
 
 	if (CPUTexelFactor >= 0)
 	{
diff --git a/Engine/Shaders/MobileBasePassPixelShader.usf b/Engine/Shaders/MobileBasePassPixelShader.usf
--- a/Engine/Shaders/MobileBasePassPixelShader.usf
+++ b/Engine/Shaders/MobileBasePassPixelShader.usf
@@ -44,6 +44,7 @@ float4 LightColorAndFalloffExponent[MAX_DYNAMIC_POINT_LIGHTS];
 #if !FULLY_ROUGH
 TextureCube ReflectionCubemap;
 SamplerState ReflectionCubemapSampler;
+float3 InvReflectionCubemapAverageBrightness;
 	#if HQ_REFLECTIONS
 	#define MAX_HQ_REFLECTIONS 3
 
@@ -155,7 +156,7 @@ float4 GetBoxCaptureVector(float3 ReflectionVector, float3 WorldPosition, float4
 }
 #endif
 
-half3 BlendReflectionCapture(FMaterialPixelParameters MaterialParameters, half Roughness, TextureCube ReflectionCube, SamplerState ReflectionSampler, half MipLevel, int ReflectionIndex, inout half BlendFactor)
+half3 BlendReflectionCapture(FMaterialPixelParameters MaterialParameters, half Roughness, TextureCube ReflectionCube, SamplerState ReflectionSampler, half MipLevel, int ReflectionIndex, half InvReflectionCubemapAverageBrightness, inout half BlendFactor)
 {
 	half4 ProjectedCaptureVector;
 #if ALLOW_CUBE_REFLECTIONS
@@ -172,6 +173,7 @@ half3 BlendReflectionCapture(FMaterialPixelParameters MaterialParameters, half R
 	// Fetch from cubemap and convert to linear HDR
 	float4 Reflection = ReflectionCube.SampleLevel(ReflectionSampler, ProjectedCaptureVector.xyz, MipLevel);
 	half3 SpecularIBL = RGBMDecode(Reflection, 16.0f) * ProjectedCaptureVector.w * BlendFactor;
+	SpecularIBL = SpecularIBL * SpecularIBL * InvReflectionCubemapAverageBrightness;
 
 	BlendFactor = BlendFactor * (1.0 - ProjectedCaptureVector.w);
 	return SpecularIBL;
@@ -185,9 +187,9 @@ half3 BlendReflectionCaptures(FMaterialPixelParameters MaterialParameters, half
 	half3 SpecularIBL = half3(0, 0, 0);
 	half BlendFactor = 1;
 
-	SpecularIBL += BlendReflectionCapture(MaterialParameters, Roughness, ReflectionCubemap, ReflectionCubemapSampler, AbsoluteSpecularMip, 0, BlendFactor);
-	SpecularIBL += BlendReflectionCapture(MaterialParameters, Roughness, ReflectionCubemap1, ReflectionCubemapSampler1, AbsoluteSpecularMip, 1, BlendFactor);
-	SpecularIBL += BlendReflectionCapture(MaterialParameters, Roughness, ReflectionCubemap2, ReflectionCubemapSampler2, AbsoluteSpecularMip, 2, BlendFactor);
+	SpecularIBL += BlendReflectionCapture(MaterialParameters, Roughness, ReflectionCubemap, ReflectionCubemapSampler, AbsoluteSpecularMip, 0, InvReflectionCubemapAverageBrightness.x, BlendFactor);
+	SpecularIBL += BlendReflectionCapture(MaterialParameters, Roughness, ReflectionCubemap1, ReflectionCubemapSampler1, AbsoluteSpecularMip, 1, InvReflectionCubemapAverageBrightness.y, BlendFactor);
+	SpecularIBL += BlendReflectionCapture(MaterialParameters, Roughness, ReflectionCubemap2, ReflectionCubemapSampler2, AbsoluteSpecularMip, 2, InvReflectionCubemapAverageBrightness.z, BlendFactor);
 
 	return SpecularIBL;
 }
@@ -210,8 +212,9 @@ half3 GetImageBasedReflectionLighting(FMaterialPixelParameters MaterialParameter
 	half AbsoluteSpecularMip = ComputeReflectionCaptureMipFromRoughness(Roughness, View.ReflectionCubemapMaxMip);
 	// Fetch from cubemap and convert to linear HDR
 	half3 SpecularIBL = RGBMDecode(ReflectionCubemap.SampleLevel(ReflectionCubemapSampler, ProjectedCaptureVector, AbsoluteSpecularMip), 16.0f);
+	SpecularIBL = SpecularIBL * SpecularIBL * InvReflectionCubemapAverageBrightness.x;
 #endif
-	SpecularIBL *= SpecularIBL;
+	
 #if WEBGL
 	// need a rgb swizzle instead of the existing rgba swizzle, we should add it if another use case comes up. 
 	return SpecularIBL.bgr;
diff --git a/Engine/Shaders/PostProcessMotionBlur.usf b/Engine/Shaders/PostProcessMotionBlur.usf
--- a/Engine/Shaders/PostProcessMotionBlur.usf
+++ b/Engine/Shaders/PostProcessMotionBlur.usf
@@ -206,7 +206,7 @@ void MainPS(
 	float2 TileJitter = ( float2( Random, Random2 ) - 0.5 ) * 0.5;
 	float2 TileUV = min( ( PixelPos - ViewportRect.xy + 0.5 ) / 16.0 + TileJitter, ViewportRect.zw ) * PostprocessInput3Size.zw;
 
-	float4 MinMaxVelocity = PostprocessInput3.SampleLevel( PostprocessInput3Sampler, TileUV, 0 );
+	float4 MinMaxVelocity = PostprocessInput3.SampleLevel(PostprocessInput3Sampler, TileUV, 0);
 	//float4 MinMaxVelocity = PostprocessInput3[ ( PixelPos - ViewportRect.xy ) / 16 ];
 
 	float2 MinVelocityPixels = MinMaxVelocity.xy * MotionBlurParameters.y;
@@ -214,7 +214,7 @@ void MainPS(
 
 	float MinVelocityLengthSqrPixels = dot( MinVelocityPixels, MinVelocityPixels );
 	float MaxVelocityLengthSqrPixels = dot( MaxVelocityPixels, MaxVelocityPixels );
-	
+
 	// TODO expose cvars
 	bool bSkipPath = MaxVelocityLengthSqrPixels < 0.25;
 	bool bFastPath = MinVelocityLengthSqrPixels > 0.4 * MaxVelocityLengthSqrPixels;
@@ -230,19 +230,18 @@ void MainPS(
 
 #if 0
 	float4 SearchVectorPixels = float4( MaxVelocityPixels, -MaxVelocityPixels );
-	float4 SearchVector = SearchVectorPixels * PostprocessInput0Size.zwzw;
 #else
 	// Clip MaxVelocity to screen rect
-	float2 InvVelocityScreen = ScreenPosToPixel.x * rcp( MaxVelocityPixels );
+	float2 InvVelocityScreen = ScreenPosToPixel.x * rcp( MaxVelocityPixels + 0.00001f );
 	float2 MinIntersect = -InvVelocityScreen - ScreenSpacePos * InvVelocityScreen;
 	float2 MaxIntersect =  InvVelocityScreen - ScreenSpacePos * InvVelocityScreen;
 	float4 FarIntersect = float4( max( MinIntersect, MaxIntersect ), max( -MinIntersect, -MaxIntersect ) );
-	float2 Intersect = saturate( min( FarIntersect.xz, FarIntersect.yw ) );
+	float2 Intersect = saturate(min(FarIntersect.xz, FarIntersect.yw) );
 
 	// +/-
 	float4 SearchVectorPixels = MaxVelocityPixels.xyxy * float4( Intersect.xx, -Intersect.yy );
-	float4 SearchVector = SearchVectorPixels * PostprocessInput0Size.zwzw;
 #endif
+	float4 SearchVector = SearchVectorPixels * PostprocessInput0Size.zwzw;
 
 	// converts pixel length to sample steps
 	float PixelToSampleScale = StepCount * rsqrt( dot( MaxVelocityPixels, MaxVelocityPixels ) );
@@ -443,7 +442,7 @@ void VisualizeMotionBlurPS(
 
 			// points into the movement direction
 			VelocityAtTileCenter = ScreenPosAtTileCenter - PrevScreen;
-		}
+		} 
 	}
 
 	// tint yellow if velocity data is stored in texture
diff --git a/Engine/Shaders/PostProcessTemporalAA.usf b/Engine/Shaders/PostProcessTemporalAA.usf
--- a/Engine/Shaders/PostProcessTemporalAA.usf
+++ b/Engine/Shaders/PostProcessTemporalAA.usf
@@ -268,7 +268,7 @@ float3 MinimalReprojectionForDebugging(float2 UV, float4 SvPosition)
 
 void SSRTemporalAAPS( 
 	noperspective 	float2 UV 						: TEXCOORD0, 
-					float3 InExposureScaleVignette 	: TEXCOORD1, 
+	noperspective	float3 InExposureScaleVignette 	: TEXCOORD1, 
 					float4 Unused2 					: TEXCOORD2,
 					float4 Unused3 					: TEXCOORD3,
 					float4 Unused4 					: TEXCOORD4,
@@ -290,7 +290,9 @@ void SSRTemporalAAPS(
 #endif
 }
 
-void DOFTemporalAAPS( noperspective float2 UV : TEXCOORD0, float3 InExposureScaleVignette : TEXCOORD1, 
+void DOFTemporalAAPS( 
+	in noperspective float2 UV : TEXCOORD0, 
+	in noperspective float3 InExposureScaleVignette : TEXCOORD1, 
 	in noperspective float4 Unused1 : TEXCOORD2,
 	in noperspective float4 Unused2 : TEXCOORD3,
 	in noperspective float4 Unused3 : TEXCOORD4,
@@ -316,7 +318,9 @@ void DOFTemporalAAPS( noperspective float2 UV : TEXCOORD0, float3 InExposureScal
 #endif
 }
 
-void LightShaftTemporalAAPS( noperspective float2 UV : TEXCOORD0, float3 InExposureScaleVignette : TEXCOORD1, 
+void LightShaftTemporalAAPS( 
+	in noperspective float2 UV : TEXCOORD0, 
+	in noperspective float3 InExposureScaleVignette : TEXCOORD1, 
 	in noperspective float4 Unused1 : TEXCOORD2,
 	in noperspective float4 Unused2 : TEXCOORD3,
 	in noperspective float4 Unused3 : TEXCOORD4,
diff --git a/Engine/Shaders/PrimitiveDistanceAccuracyPixelShader.usf b/Engine/Shaders/PrimitiveDistanceAccuracyPixelShader.usf
--- a/Engine/Shaders/PrimitiveDistanceAccuracyPixelShader.usf
+++ b/Engine/Shaders/PrimitiveDistanceAccuracyPixelShader.usf
@@ -25,7 +25,7 @@ void Main(
 	ResolvedView = ResolveView();
 #endif
 
-	float3 Result = float3(.1, .1, .1);
+	float3 Result = float3(UNDEFINED_ACCURACY, UNDEFINED_ACCURACY, UNDEFINED_ACCURACY);
 
 	if (CPULogDistance >= 0)
 	{
diff --git a/Engine/Shaders/Random.usf b/Engine/Shaders/Random.usf
--- a/Engine/Shaders/Random.usf
+++ b/Engine/Shaders/Random.usf
@@ -72,15 +72,10 @@ float RandBBSfloat(float seed)
 // ~5 ALU operations (3 %, 2 *)
 uint RandBBSuint24(uint seed)
 {
-#if (ES2_PROFILE)
-	// no integer mod in ES2
-	return BBS_PRIME24 * RandBBSfloat(float(seed));
-#else
-	uint s = seed % BBS_PRIME24;
-	s = (s * s) % BBS_PRIME24;
-	s = (s * s) % BBS_PRIME24;
+	uint s = Mod(seed, BBS_PRIME24);
+	s = Mod(s * s, BBS_PRIME24);
+	s = Mod(s * s, BBS_PRIME24);
 	return s;
-#endif
 }
 
 // BBS random number generator - 32-bit uint version
@@ -90,15 +85,10 @@ uint RandBBSuint24(uint seed)
 // ~5 ALU operations (3 %, 2 *)
 uint RandBBSuint32(uint seed)
 {
-#if (ES2_PROFILE)
-	// no integer mod in ES2
-	return BBS_PRIME32 * RandBBSfloat(float(seed));
-#else
-	uint s = seed % BBS_PRIME32;
-	s = (s * s) % BBS_PRIME32;
-	s = (s * s) % BBS_PRIME32;
+	uint s = Mod(seed, BBS_PRIME32);
+	s = Mod(s * s, BBS_PRIME32);
+	s = Mod(s * s, BBS_PRIME32);
 	return s;
-#endif
 }
 
 // 3D random number generator inspired by PCGs (permuted congruential generator)
diff --git a/Engine/Shaders/ReflectionEnvironmentComputeShaders.usf b/Engine/Shaders/ReflectionEnvironmentComputeShaders.usf
--- a/Engine/Shaders/ReflectionEnvironmentComputeShaders.usf
+++ b/Engine/Shaders/ReflectionEnvironmentComputeShaders.usf
@@ -48,7 +48,12 @@ groupshared uint TileReflectionCaptureIndices[MAX_CAPTURES];
 groupshared uint SortedTileReflectionCaptureIndices[MAX_CAPTURES];
 
 #if !ATOMIC_REDUCTION
-groupshared float TileZ[THREADGROUP_TOTALSIZE];
+	#if THREADGROUP_TOTALSIZE < 107
+		#define TILE_Z_SIZE 107
+	#else
+		#define TILE_Z_SIZE THREADGROUP_TOTALSIZE
+	#endif
+groupshared float TileZ[TILE_Z_SIZE];
 #endif
 
 void ComputeTileMinMax(uint ThreadIndex, float SceneDepth, out float MinTileZ, out float MaxTileZ, out float MinTileZ2, out float MaxTileZ2)
@@ -410,7 +415,7 @@ float CountOverlap( float3 WorldPosition )
 	return Overlap;
 }
 
-void GatherRadiance(inout float4 Color, float3 WorldPosition, float3 RayDirection, float Roughness, float2 ScreenPosition, float IndirectIrradiance, float NoV, uint ShadingModelID)
+float3 GatherRadiance(float CompositeAlpha, float3 WorldPosition, float3 RayDirection, float Roughness, float2 ScreenPosition, float IndirectIrradiance, float NoV, uint ShadingModelID)
 {
 	// Indirect occlusion from DFAO, which should be applied to reflection captures and skylight specular, but not SSR
 	float IndirectSpecularOcclusion = 1.0f;
@@ -425,15 +430,16 @@ void GatherRadiance(inout float4 Color, float3 WorldPosition, float3 RayDirectio
 #endif
 
 	float Mip = ComputeReflectionCaptureMipFromRoughness(Roughness, View.ReflectionCubemapMaxMip);
-
 	uint NumCapturesAffectingTile = TileNumReflectionCaptures;
+	float4 ImageBasedReflections = float4(0, 0, 0, CompositeAlpha);
+	float2 CompositedAverageBrightness = float2(0.0f, 1.0f);
 
 	// Accumulate reflections from captures affecting this tile, applying largest captures first so that the smallest ones display on top
 	LOOP
 	for (uint TileCaptureIndex = 0; TileCaptureIndex < NumCapturesAffectingTile; TileCaptureIndex++) 
 	{
 		BRANCH
-		if( Color.a < 0.001 )
+		if (ImageBasedReflections.a < 0.001)
 		{
 			break;
 		}
@@ -455,33 +461,32 @@ void GatherRadiance(inout float4 Color, float3 WorldPosition, float3 RayDirectio
 		if (CaptureVectorLength < CapturePositionAndRadius.w)
 		{
 			float3 ProjectedCaptureVector = RayDirection;
+			float4 CaptureOffsetAndAverageBrightness = ReflectionCapture.CaptureOffsetAndAverageBrightness[CaptureIndex];
 
 			// Fade out based on distance to capture
 			float DistanceAlpha = 0;
 			
 			#define PROJECT_ONTO_SHAPE 1
 			#if PROJECT_ONTO_SHAPE
+				#if HAS_BOX_CAPTURES
+					#if HAS_SPHERE_CAPTURES
+					// Box
+					BRANCH if (CaptureProperties.b > 0)
+					#endif
+					{
+						ProjectedCaptureVector = GetLookupVectorForBoxCapture(RayDirection, WorldPosition, CapturePositionAndRadius, ReflectionCapture.BoxTransform[CaptureIndex], ReflectionCapture.BoxScales[CaptureIndex], CaptureOffsetAndAverageBrightness.xyz, DistanceAlpha);
+					}
+				#endif
 
-#if HAS_BOX_CAPTURES
-
-#	if HAS_SPHERE_CAPTURES
-				// Box
-				BRANCH if (CaptureProperties.b > 0)
-#	endif //HAS_SPHERE_CAPTURES
-				{
-					ProjectedCaptureVector = GetLookupVectorForBoxCapture(RayDirection, WorldPosition, CapturePositionAndRadius, ReflectionCapture.BoxTransform[CaptureIndex], ReflectionCapture.BoxScales[CaptureIndex], ReflectionCapture.CaptureOffset[CaptureIndex].xyz, DistanceAlpha);
-				}
-#endif //HAS_BOX_CAPTURES
-
-#if HAS_SPHERE_CAPTURES
-				// Sphere
-#	if HAS_BOX_CAPTURES
-				else
-#	endif //HAS_BOX_CAPTURES
-				{
-					ProjectedCaptureVector = GetLookupVectorForSphereCapture(RayDirection, WorldPosition, CapturePositionAndRadius, NormalizedDistanceToCapture, ReflectionCapture.CaptureOffset[CaptureIndex].xyz, DistanceAlpha);
-				}
-#endif //HAS_SPHERE_CAPTURES
+				#if HAS_SPHERE_CAPTURES
+					// Sphere
+					#if HAS_BOX_CAPTURES
+					else
+					#endif
+					{
+						ProjectedCaptureVector = GetLookupVectorForSphereCapture(RayDirection, WorldPosition, CapturePositionAndRadius, NormalizedDistanceToCapture, CaptureOffsetAndAverageBrightness.xyz, DistanceAlpha);
+					}
+				#endif
 			#else 
 				DistanceAlpha = 1.0;
 			#endif //PROJECT_ONTO_SHAPE
@@ -489,44 +494,54 @@ void GatherRadiance(inout float4 Color, float3 WorldPosition, float3 RayDirectio
 			float CaptureArrayIndex = CaptureProperties.g;
 
 			{
-				float4 Sample = ReflectionEnvironmentColorTexture.SampleLevel( ReflectionEnvironmentColorSampler, float4( ProjectedCaptureVector, CaptureArrayIndex ), Mip );
-
-				#if USE_LIGHTMAPS
-					// We have high frequency directional data but low frequency spatial data in the envmap.
-					// We have high frequency spatial data but low frequency directional data in the lightmap.
-					// So, we combine the two for the best of both. This is done by removing the low spatial frequencies from the envmap and replacing them with the lightmap data.
-					// This is only done with luma so as to not get odd color shifting.
-					// Note: make sure this matches the lightmap mixing done for translucency (BasePassPixelShader.usf)
-					Sample.rgb *= IndirectIrradiance;
-				#endif
+				float4 Sample = ReflectionEnvironmentColorTexture.SampleLevel(ReflectionEnvironmentColorSampler, float4(ProjectedCaptureVector, CaptureArrayIndex), Mip);
 
 				Sample.rgb *= CaptureProperties.r;
 				Sample *= DistanceAlpha;
 
 				// Under operator (back to front)
-				Color.rgb += Sample.rgb * Color.a * IndirectSpecularOcclusion;
-				Color.a *= 1 - Sample.a;
+				ImageBasedReflections.rgb += Sample.rgb * ImageBasedReflections.a * IndirectSpecularOcclusion;
+				ImageBasedReflections.a *= 1 - Sample.a;
+
+				float AverageBrightness = CaptureOffsetAndAverageBrightness.w;
+				CompositedAverageBrightness.x += AverageBrightness * DistanceAlpha * CompositedAverageBrightness.y;
+				CompositedAverageBrightness.y *= 1 - DistanceAlpha;
 			}
 		}
 	}
 
 #if HAS_SKYLIGHT
-	BRANCH
-	if( SkyLightParameters.y > 0 && Color.a >= 0.001 )
+
+	BRANCH  
+	if (SkyLightParameters.y > 0)
 	{
+		float SkyAverageBrightness = 1.0f;
+		float3 SkyLighting = GetSkyLightReflectionSupportingBlend(RayDirection, Roughness, SkyAverageBrightness);
+
 		// Normalize for static skylight types which mix with lightmaps
 		bool bNormalize = SkyLightParameters.z < 1 && USE_LIGHTMAPS;
-		float3 SkyLighting = GetSkyLightReflectionSupportingBlend( RayDirection, Roughness, bNormalize );
 
 		FLATTEN
 		if (bNormalize)
 		{
-			SkyLighting *= IndirectIrradiance;
+			ImageBasedReflections.rgb += ImageBasedReflections.a * SkyLighting * IndirectSpecularOcclusion;
+			CompositedAverageBrightness.x += SkyAverageBrightness * CompositedAverageBrightness.y;
+		}
+		else
+		{
+			ExtraIndirectSpecular += SkyLighting * IndirectSpecularOcclusion;
 		}
-
-		Color.rgb += Color.a * (SkyLighting * IndirectSpecularOcclusion + ExtraIndirectSpecular);
 	}
 #endif
+
+#if USE_LIGHTMAPS
+	// Note: make sure this matches the lightmap mixing done for translucency (BasePassPixelShader.usf)
+	ImageBasedReflections.rgb *= ComputeMixingWeight(IndirectIrradiance, CompositedAverageBrightness.x, Roughness);
+#endif
+
+	ImageBasedReflections.rgb += ImageBasedReflections.a * ExtraIndirectSpecular;
+
+	return ImageBasedReflections.rgb;
 }
 
 
@@ -586,8 +601,6 @@ void ReflectionEnvironmentTiledDeferredMain(
 	BRANCH
 	if( GBuffer.ShadingModelID != SHADINGMODELID_UNLIT && GBuffer.ShadingModelID != SHADINGMODELID_HAIR )
 	{		
-
-
 		float3 N = GBuffer.WorldNormal;
 		float3 V = -CameraToPixel;
 		float3 R = 2 * dot( V, N ) * N - V;
@@ -623,7 +636,7 @@ void ReflectionEnvironmentTiledDeferredMain(
 		Color.a *= SpecularOcclusion;
 
 		//bottom for clearcoat or the only reflection.
-		GatherRadiance(Color, WorldPosition, R, GBuffer.Roughness, ScreenPosition, IndirectIrradiance, NoV, GBuffer.ShadingModelID);
+		Color.rgb += GatherRadiance(Color.a, WorldPosition, R, GBuffer.Roughness, ScreenPosition, IndirectIrradiance, NoV, GBuffer.ShadingModelID);
 
 		BRANCH
 		if( GBuffer.ShadingModelID == SHADINGMODELID_CLEAR_COAT )
@@ -652,7 +665,7 @@ void ReflectionEnvironmentTiledDeferredMain(
 			Color.a *= SpecularOcclusion;
 
 			float3 TopLayerR = 2 * dot( V, N ) * N - V;
-			GatherRadiance(Color, WorldPosition, TopLayerR, ClearCoatRoughness, ScreenPosition, IndirectIrradiance, NoV, GBuffer.ShadingModelID);
+			Color.rgb += GatherRadiance(Color.a, WorldPosition, TopLayerR, ClearCoatRoughness, ScreenPosition, IndirectIrradiance, NoV, GBuffer.ShadingModelID);
 		}
 		else
 		{
diff --git a/Engine/Shaders/ReflectionEnvironmentShaders.usf b/Engine/Shaders/ReflectionEnvironmentShaders.usf
--- a/Engine/Shaders/ReflectionEnvironmentShaders.usf
+++ b/Engine/Shaders/ReflectionEnvironmentShaders.usf
@@ -75,8 +75,9 @@ void CopySceneColorToCubeFaceColorPS(
 		Alpha = 1 - smoothstep(.8f * MaxDistance, MaxDistance, RadialDistance);
 	}
 
-	// Pre-multiplied alpha for correct filtering
-	OutColor = float4(SceneColor * Alpha, Alpha);
+	// We need pre-multiplied alpha for correct filtering
+	// However, we need to compute average brightness before masking out sky areas, so premultiplying happens later
+	OutColor = float4(SceneColor, Alpha);
 }
 
 float3 GetCubemapVector(float2 ScaledUVs)
@@ -136,6 +137,8 @@ void CopyCubemapToCubeFaceColorPS(
 			OutColor.rgb = lerp(OutColor.rgb, LowerHemisphereColor.rgb, LowerHemisphereColor.a);
 		}
 	}
+
+	OutColor.a = 1;
 }
 
 int SourceMipIndex;
@@ -177,15 +180,13 @@ void ComputeBrightnessMain(
 	AverageColor += TextureCubeSampleLevel(ReflectionEnvironmentColorTexture, ReflectionEnvironmentColorSampler, float3(0, 0, 1), NumCaptureArrayMips - 1).rgb;
 	AverageColor += TextureCubeSampleLevel(ReflectionEnvironmentColorTexture, ReflectionEnvironmentColorSampler, float3(0, 0, -1), NumCaptureArrayMips - 1).rgb;
 
-	OutColor = dot(AverageColor / 6, .3333f);
+	OutColor = dot(AverageColor / 6, .3333f); 
 }
 
 #endif
 
 #define NUM_FILTER_SAMPLES 1024
 
-Texture2D AverageBrightnessTexture;
-SamplerState AverageBrightnessSampler;
 float CubemapMaxMip;
 
 void FilterPS(
@@ -233,13 +234,6 @@ void FilterPS(
 	OutColor = SampleCubemap(CubeCoordinates);
 
 #endif
-
-#if NORMALIZE && ALLOW_STATIC_LIGHTING
-	// Remove low frequency
-	// TODO match the frequency removed to the one added. In other words divide out 2nd order SH instead of 1st.
-	float LowFrequencyBrightness = Texture2DSample(AverageBrightnessTexture, AverageBrightnessSampler, float2(0, 0)).r;
-	OutColor.rgb /= max(LowFrequencyBrightness, 0.01f);
-#endif
 }
 
 float4 CoefficientMask0;
@@ -333,7 +327,7 @@ float4 CapturePositionAndRadius;
 float4 CaptureProperties;
 float4x4 CaptureBoxTransform;
 float4 CaptureBoxScales;
-float3 CaptureOffset;
+float4 CaptureOffsetAndAverageBrightness;
 
 #if FEATURE_LEVEL >= FEATURE_LEVEL_SM5
 TextureCubeArray ReflectionEnvironmentColorTextureArray;
@@ -371,14 +365,15 @@ void StandardDeferredReflectionPS(
 			float CaptureVectorLength = sqrt(dot(CaptureVector, CaptureVector));
 			float NormalizedDistanceToCapture = saturate(CaptureVectorLength / CapturePositionAndRadius.w);
 
-			ProjectedCaptureVector = GetLookupVectorForSphereCapture(ReflectionVector, WorldPosition, CapturePositionAndRadius, NormalizedDistanceToCapture, CaptureOffset, DistanceAlpha);
+			ProjectedCaptureVector = GetLookupVectorForSphereCapture(ReflectionVector, WorldPosition, CapturePositionAndRadius, NormalizedDistanceToCapture, CaptureOffsetAndAverageBrightness.xyz, DistanceAlpha);
 		#elif BOX_CAPTURE
-			ProjectedCaptureVector = GetLookupVectorForBoxCapture(ReflectionVector, WorldPosition, CapturePositionAndRadius, CaptureBoxTransform, CaptureBoxScales, CaptureOffset, DistanceAlpha);
+			ProjectedCaptureVector = GetLookupVectorForBoxCapture(ReflectionVector, WorldPosition, CapturePositionAndRadius, CaptureBoxTransform, CaptureBoxScales, CaptureOffsetAndAverageBrightness.xyz, DistanceAlpha);
 		#endif
 		
 		BRANCH
 		if (DistanceAlpha > 0)
 		{
+			float AverageBrightness = CaptureOffsetAndAverageBrightness.w;
 			float AbsoluteSpecularMip = ComputeReflectionCaptureMipFromRoughness(InGBufferData.Roughness, View.ReflectionCubemapMaxMip);
 
 #if FEATURE_LEVEL >= FEATURE_LEVEL_SM5
@@ -387,6 +382,10 @@ void StandardDeferredReflectionPS(
 			float4 Reflection = TextureCubeSampleLevel(ReflectionEnvironmentColorTexture, ReflectionEnvironmentColorSampler, ProjectedCaptureVector, AbsoluteSpecularMip);
 #endif
 
+			#if ALLOW_STATIC_LIGHTING
+				Reflection.rgb /= AverageBrightness;
+			#endif
+			
 			Reflection.rgb *= CaptureProperties.r;
 
 			// Composite using the influence shape's alpha, which allows it to only override where it has valid data
@@ -450,18 +449,7 @@ void ReflectionApplyPS(
 		SpecularLighting.rgb *= IndirectSpecularOcclusion;
 
 #if ALLOW_STATIC_LIGHTING
-		FLATTEN
-		if( View.UseLightmaps > 0 )
-		{
-			// We have high frequency directional data but low frequency spatial data in the envmap.
-			// We have high frequency spatial data but low frequency directional data in the lightmap.
-			// So, we combine the two for the best of both. This is done by removing the low spatial frequencies from the envmap and replacing them with the lightmap data.
-			// This is only done with luma so as to not get odd color shifting.
-			// Note: make sure this matches the lightmap mixing done for translucency (BasePassPixelShader.usf)
-
-			// NOTE: this breaks LPV specular if r.allowstaticlighting is enabled
-			SpecularLighting.rgb *= IndirectIrradiance;
-		}
+		SpecularLighting.rgb *= IndirectIrradiance;
 #endif
 #endif
 
@@ -480,9 +468,8 @@ void ReflectionApplyPS(
 		if (SpecularLighting.a > .001f)
 		{			
 			float3 SkyLighting = float3(1,1,1);
+			float SkyAverageBrightness = 1.0f;
 
-			// Normalize for static skylight types which mix with lightmaps
-			bool bNormalize = SkyLightParameters.z < 1 && View.UseLightmaps > 0 && !ALLOW_STATIC_LIGHTING;
 			BRANCH
 			if( GBuffer.ShadingModelID == SHADINGMODELID_CLEAR_COAT )
 			{
@@ -494,7 +481,7 @@ void ReflectionApplyPS(
 				ReflectionVector = 2 * dot(V, ClearCoatUnderNormal) * ClearCoatUnderNormal - V;
 #endif				
 
-				SkyLighting = GetSkyLightReflectionSupportingBlend(ReflectionVector, GBuffer.Roughness, bNormalize);
+				SkyLighting = GetSkyLightReflectionSupportingBlend(ReflectionVector, GBuffer.Roughness, SkyAverageBrightness);
 				SkyLighting *= SpecularColor;
 
 				const float ClearCoat			= GBuffer.CustomData.x;
@@ -514,12 +501,12 @@ void ReflectionApplyPS(
 				#endif				
 
 				float3 TopLayerR = 2 * dot(V, GBuffer.WorldNormal) * N - V;
-				float3 SampleColor = GetSkyLightReflectionSupportingBlend(TopLayerR, ClearCoatRoughness, bNormalize);
+				float3 SampleColor = GetSkyLightReflectionSupportingBlend(TopLayerR, ClearCoatRoughness, SkyAverageBrightness);
 				SkyLighting += SampleColor * F;
 			}
 			else
 			{
-				SkyLighting = GetSkyLightReflectionSupportingBlend(ReflectionVector, GBuffer.Roughness, bNormalize);
+				SkyLighting = GetSkyLightReflectionSupportingBlend(ReflectionVector, GBuffer.Roughness, SkyAverageBrightness);
 				SkyLighting *= SpecularColor;
 
 			}
@@ -527,10 +514,13 @@ void ReflectionApplyPS(
 		    float DirectionalOcclusion = ScreenSpaceData.DirectionalOcclusion.g;
 		    SkyLighting *= DirectionalOcclusion;
 
+			// Normalize for static skylight types which mix with lightmaps
+			bool bNormalize = SkyLightParameters.z < 1 && ALLOW_STATIC_LIGHTING;
+
 			FLATTEN
 			if (bNormalize)
 			{
-				SkyLighting *= IndirectIrradiance;
+				SkyLighting *= IndirectIrradiance / max(SkyAverageBrightness, .001f);
 			}
 
 			SpecularLighting.rgb += SpecularLighting.a * (SkyLighting * IndirectSpecularOcclusion + ExtraIndirectSpecular * SpecularColor);
diff --git a/Engine/Shaders/ReflectionEnvironmentShared.usf b/Engine/Shaders/ReflectionEnvironmentShared.usf
--- a/Engine/Shaders/ReflectionEnvironmentShared.usf
+++ b/Engine/Shaders/ReflectionEnvironmentShared.usf
@@ -198,4 +198,16 @@ float3 GetLookupVectorForSphereCapture(float3 ReflectionVector, float3 WorldPosi
 		DistanceAlpha = 1 - x*x*(3 - 2*x);
 	}
 	return ProjectedCaptureVector;
+}
+
+float ComputeMixingWeight(float IndirectIrradiance, float AverageBrightness, float Roughness)
+{
+	// Mirror surfaces should have no mixing, so they match reflections from other sources (SSR, planar reflections)
+	float MixingAlpha = smoothstep(0, 1, saturate(Roughness * View.ReflectionEnvironmentRoughnessMixingScaleBias.x + View.ReflectionEnvironmentRoughnessMixingScaleBias.y));
+
+	// We have high frequency directional data but low frequency spatial data in the envmap.
+	// We have high frequency spatial data but low frequency directional data in the lightmap.
+	// So, we combine the two for the best of both. This is done by removing the low spatial frequencies from the envmap and replacing them with the lightmap data.
+	// This is only done with luma so as to not get odd color shifting.
+	return lerp(1.0f, IndirectIrradiance / max(AverageBrightness, .0001f), MixingAlpha);
 }
\ No newline at end of file
diff --git a/Engine/Shaders/SimpleElementColorChannelMaskPixelShader.usf b/Engine/Shaders/SimpleElementColorChannelMaskPixelShader.usf
--- a/Engine/Shaders/SimpleElementColorChannelMaskPixelShader.usf
+++ b/Engine/Shaders/SimpleElementColorChannelMaskPixelShader.usf
@@ -14,14 +14,14 @@ float Gamma;
 
 bool bAlphaOnly;
 
+#define WRITE_TO_GBUFFER (FEATURE_LEVEL >= FEATURE_LEVEL_SM4 && !FORWARD_SHADING)
+
 void Main(
 	in float2 TextureCoordinate : TEXCOORD0,
 	in float4 Color : TEXCOORD1,
 	out float4 OutColor : SV_Target0
-#if FEATURE_LEVEL >= FEATURE_LEVEL_SM4
+#if WRITE_TO_GBUFFER
 	,out float4 OutWorldNormal : SV_Target1 
-	,out float4 OutReflectionWorldNormal : SV_Target2 
-	,out float4 OutSpecularColorAndPower : SV_Target3
 #endif
 	)
 {
@@ -44,10 +44,8 @@ void Main(
 	
 	OutColor = RETURN_COLOR(FinalColor);	
 
-#if FEATURE_LEVEL >= FEATURE_LEVEL_SM4
+#if WRITE_TO_GBUFFER
 	// Set the G buffer bits that indicate that deferred lighting and image reflections are not enabled
 	OutWorldNormal = 0;
-	OutReflectionWorldNormal = 0;
-	OutSpecularColorAndPower = 0;
 #endif
 }
\ No newline at end of file
diff --git a/Engine/Shaders/SimpleElementNormalMapPixelShader.usf b/Engine/Shaders/SimpleElementNormalMapPixelShader.usf
--- a/Engine/Shaders/SimpleElementNormalMapPixelShader.usf
+++ b/Engine/Shaders/SimpleElementNormalMapPixelShader.usf
@@ -7,6 +7,7 @@
 
 #include "Common.usf"
 
+#define WRITE_TO_GBUFFER (FEATURE_LEVEL >= FEATURE_LEVEL_SM4 && !FORWARD_SHADING)
 
 Texture2D NormalMapTexture;
 SamplerState NormalMapTextureSampler;
@@ -14,9 +15,8 @@ SamplerState NormalMapTextureSampler;
 void Main(
 	in float2 TextureCoordinate : TEXCOORD0,
 	out float4 OutColor : SV_Target0
-#if FEATURE_LEVEL >= FEATURE_LEVEL_SM4
+#if WRITE_TO_GBUFFER
 	,out float4 OutWorldNormal : SV_Target1 
-	,out float4 OutSpecularColorAndPower : SV_Target2
 #endif
 	)
 {
@@ -25,9 +25,8 @@ void Main(
 	const float4 RescaledNormal = float4(Normal.xyz * 0.5 + 0.5, 1);
 	OutColor = RETURN_COLOR(RescaledNormal);	
 
-#if FEATURE_LEVEL >= FEATURE_LEVEL_SM4
+#if WRITE_TO_GBUFFER
 	// Set the G buffer bits that indicate that deferred lighting and image reflections are not enabled
 	OutWorldNormal = 0;
-	OutSpecularColorAndPower = 0;
 #endif
 }
diff --git a/Engine/Shaders/SimpleElementPixelShader.usf b/Engine/Shaders/SimpleElementPixelShader.usf
--- a/Engine/Shaders/SimpleElementPixelShader.usf
+++ b/Engine/Shaders/SimpleElementPixelShader.usf
@@ -11,6 +11,7 @@
 #include "IESLightProfilesCommon.usf"
 
 #define ENABLE_EDITOR_COMPOSITING (FEATURE_LEVEL >= FEATURE_LEVEL_SM4 || MOBILE_EMULATION)
+#define WRITE_TO_SHADING_MODEL (FEATURE_LEVEL >= FEATURE_LEVEL_SM4 && !FORWARD_SHADING)
 
 Texture2D InTexture;
 SamplerState InTextureSampler;
@@ -106,7 +107,7 @@ void Main(
 	in float4 Color : TEXCOORD1,
 	in float4 HitProxyId : TEXCOORD2,
 	out float4 OutColor : SV_Target0
-#if FEATURE_LEVEL >= FEATURE_LEVEL_SM4
+#if WRITE_TO_SHADING_MODEL
 	,out float4 OutWorldNormal : SV_Target1 
 #endif
 #if ENABLE_EDITOR_COMPOSITING
@@ -123,7 +124,7 @@ void Main(
 	ReplicateChannelSimpleElementShader(BaseColor);
 	OutColor = RETURN_COLOR(BaseColor * Color);	
 
-#if FEATURE_LEVEL >= FEATURE_LEVEL_SM4
+#if WRITE_TO_SHADING_MODEL
 	// Set the G buffer bits that indicate unlit
 	OutWorldNormal = 0;
 #endif
@@ -136,7 +137,7 @@ void AlphaOnlyMain(
 	in float4 Color : TEXCOORD1,
 	in float4 HitProxyId : TEXCOORD2,
 	out float4 OutColor : SV_Target0
-#if FEATURE_LEVEL >= FEATURE_LEVEL_SM4
+#if WRITE_TO_SHADING_MODEL
 	,out float4 OutWorldNormal : SV_Target1 
 #endif
 #if ENABLE_EDITOR_COMPOSITING
@@ -152,7 +153,7 @@ void AlphaOnlyMain(
 	TmpColor.a *= Texture2DSample_A8(InTexture, InTextureSampler, TextureCoordinate);
 	OutColor = RETURN_COLOR(TmpColor);
 
-#if FEATURE_LEVEL >= FEATURE_LEVEL_SM4
+#if WRITE_TO_SHADING_MODEL
 	// Set the G buffer bits that indicate unlit
 	OutWorldNormal = 0;
 #endif
@@ -175,7 +176,7 @@ void GammaMain(
 #if ENABLE_EDITOR_COMPOSITING
 	,in float4 SvPosition : SV_Position
 #endif
-#if FEATURE_LEVEL >= FEATURE_LEVEL_SM4
+#if WRITE_TO_SHADING_MODEL
 	,out float4 OutWorldNormal : SV_Target1
 #endif
 	)
@@ -194,7 +195,7 @@ void GammaMain(
 	}
 	OutColor = RETURN_COLOR(OutColor);
 
-#if FEATURE_LEVEL >= FEATURE_LEVEL_SM4
+#if WRITE_TO_SHADING_MODEL
 	// Set the G buffer bits that indicate unlit
 	OutWorldNormal = 0;
 #endif
@@ -210,7 +211,7 @@ void GammaAlphaOnlyMain(
 #if ENABLE_EDITOR_COMPOSITING
 	,in float4 SvPosition : SV_Position
 #endif
-#if FEATURE_LEVEL >= FEATURE_LEVEL_SM4
+#if WRITE_TO_SHADING_MODEL
 	,out float4 OutWorldNormal : SV_Target1
 #endif
 	)
@@ -228,7 +229,7 @@ void GammaAlphaOnlyMain(
 	}
 	OutColor = RETURN_COLOR(TmpColor);
 
-#if FEATURE_LEVEL >= FEATURE_LEVEL_SM4
+#if WRITE_TO_SHADING_MODEL
 	// Set the G buffer bits that indicate unlit
 	OutWorldNormal = 0;
 #endif
@@ -245,8 +246,8 @@ void GammaMaskedMain(
 #if ENABLE_EDITOR_COMPOSITING
 	,in float4 SvPosition : SV_Position
 #endif
-#if FEATURE_LEVEL >= FEATURE_LEVEL_SM4
-	,out float4 OutWorldNormal : SV_Target2
+#if WRITE_TO_SHADING_MODEL
+	,out float4 OutWorldNormal : SV_Target1
 #endif
 	)
 {
@@ -265,7 +266,7 @@ void GammaMaskedMain(
 	}
 	OutColor = RETURN_COLOR(OutColor);
 
-#if FEATURE_LEVEL >= FEATURE_LEVEL_SM4
+#if WRITE_TO_SHADING_MODEL
 	// Set the G buffer bits that indicate unlit
 	OutWorldNormal = 0;
 #endif
@@ -378,7 +379,7 @@ void CubemapTextureProperties(
 	in float2 TextureCoordinate : TEXCOORD0,
 	in float4 InterpolatedColor : TEXCOORD1,
 	out float4 OutColor : SV_Target0
-#if FEATURE_LEVEL >= FEATURE_LEVEL_SM4
+#if WRITE_TO_SHADING_MODEL
 	,out float4 OutWorldNormal : SV_Target1 
 #endif
 	)
@@ -412,7 +413,7 @@ void CubemapTextureProperties(
 	
 	OutColor = RETURN_COLOR(FinalColor);	
 
-#if FEATURE_LEVEL >= FEATURE_LEVEL_SM4
+#if WRITE_TO_SHADING_MODEL
 	// Set the G buffer bits that indicate unlit
 	OutWorldNormal = 0;
 #endif
diff --git a/Engine/Shaders/SimpleElementTexture2DPreviewPixelShader.usf b/Engine/Shaders/SimpleElementTexture2DPreviewPixelShader.usf
--- a/Engine/Shaders/SimpleElementTexture2DPreviewPixelShader.usf
+++ b/Engine/Shaders/SimpleElementTexture2DPreviewPixelShader.usf
@@ -8,6 +8,8 @@
 #include "ColorUtils.usf"
 #include "GammaCorrectionCommon.usf"
 
+#define WRITE_TO_GBUFFER (FEATURE_LEVEL >= FEATURE_LEVEL_SM4 && !FORWARD_SHADING)
+
 Texture2D InTexture;
 SamplerState InTextureSampler;
 
@@ -23,10 +25,8 @@ void Main(
 	in float2 TextureCoordinate : TEXCOORD0,
 	in float4 Color : TEXCOORD1,
 	out float4 OutColor : SV_Target0
-#if FEATURE_LEVEL >= FEATURE_LEVEL_SM4
+#if WRITE_TO_GBUFFER
 	,out float4 OutWorldNormal : SV_Target1 
-	,out float4 OutReflectionWorldNormal : SV_Target2 
-	,out float4 OutSpecularColorAndPower : SV_Target3
 #endif
 	)
 {
@@ -72,10 +72,8 @@ void Main(
 		 
 	OutColor = FinalColor;
 
-#if FEATURE_LEVEL >= FEATURE_LEVEL_SM4
+#if WRITE_TO_GBUFFER
 	// Set the G buffer bits that indicate that deferred lighting and image reflections are not enabled
 	OutWorldNormal = 0;
-	OutReflectionWorldNormal = 0;
-	OutSpecularColorAndPower = 0;
 #endif
 }
\ No newline at end of file
diff --git a/Engine/Source/Developer/LauncherServices/Private/Profiles/LauncherProfile.h b/Engine/Source/Developer/LauncherServices/Private/Profiles/LauncherProfile.h
--- a/Engine/Source/Developer/LauncherServices/Private/Profiles/LauncherProfile.h
+++ b/Engine/Source/Developer/LauncherServices/Private/Profiles/LauncherProfile.h
@@ -1594,8 +1594,18 @@ public:
 		ForceClose = Object.GetBoolField("ForceClose");
 		Timeout = (uint32)Object.GetNumberField("Timeout");
 		Compressed = Object.GetBoolField("Compressed");
-		EncryptIniFiles = Object.GetBoolField("EncryptIniFiles");
-		ForDistribution = Object.GetBoolField("ForDistribution");
+
+		if (Version >= LAUNCHERSERVICES_ADDEDENCRYPTINIFILES)
+		{
+			EncryptIniFiles = Object.GetBoolField("EncryptIniFiles");
+			ForDistribution = Object.GetBoolField("ForDistribution");
+		}
+		else
+		{
+			EncryptIniFiles = false;
+			ForDistribution = false;
+		}
+
 		DefaultDeployPlatform = *(Object.GetStringField("DeployPlatform"));
 		NumCookersToSpawn = (int32)Object.GetNumberField("NumCookersToSpawn");
 		bSkipCookingEditorContent = Object.GetBoolField("SkipCookingEditorContent");
diff --git a/Engine/Source/Developer/TextureFormatIntelISPCTexComp/Private/TextureFormatIntelISPCTexComp.cpp b/Engine/Source/Developer/TextureFormatIntelISPCTexComp/Private/TextureFormatIntelISPCTexComp.cpp
--- a/Engine/Source/Developer/TextureFormatIntelISPCTexComp/Private/TextureFormatIntelISPCTexComp.cpp
+++ b/Engine/Source/Developer/TextureFormatIntelISPCTexComp/Private/TextureFormatIntelISPCTexComp.cpp
@@ -37,7 +37,7 @@ static FName GSupportedTextureFormatNames[] =
 /**
  * BC6H Compression function
  */
-static void IntelBC6HCompressScans(bc6h_enc_settings* pEncSettings, FImage* pInImage, FCompressedImage2D* pOutImage, int yStart, int yEnd)
+static void IntelBC6HCompressScans(bc6h_enc_settings* pEncSettings, FImage* pInImage, FCompressedImage2D* pOutImage, int yStart, int yEnd, int OutSliceSize)
 {
 	check(pInImage->Format == ERawImageFormat::RGBA16F);
 	check((yStart % 4) == 0);
@@ -55,7 +55,12 @@ static void IntelBC6HCompressScans(bc6h_enc_settings* pEncSettings, FImage* pInI
 	insurface.stride	= pInImage->SizeX * 8;
 
 	pOutTexels += ((yStart + 3) / 4) * ((pInImage->SizeX + 3) / 4) * 16;
-	CompressBlocksBC6H(&insurface, pOutTexels, pEncSettings);
+	for (int Slice = 0 ; Slice < pInImage->NumSlices ; Slice++)
+	{
+		CompressBlocksBC6H(&insurface, pOutTexels, pEncSettings);
+		pOutTexels += OutSliceSize;
+		insurface.ptr += 8 * pInImage->SizeX * pInImage->SizeY;
+	}
 }
 
 /**
@@ -96,6 +101,54 @@ static void IntelBC7CompressScans(bc7_enc_settings* pEncSettings, FImage* pInIma
 	CompressBlocksBC7(&insurface, pOutTexels, pEncSettings);
 }
 
+
+/** Expands any dimension to multiple of 4.  Copies the nearest row/column */
+static bool ExpandTo4PxMultiple(const FImage& InImage, FImage& OutImage)
+{	
+	// Calculate how many times to duplicate each row or column
+	const int32 FillX = InImage.SizeX % 4 > 0 ? 4 - InImage.SizeX % 4 : 0;
+	const int32 FillY = InImage.SizeY % 4 > 0 ? 4 - InImage.SizeY % 4 : 0;
+	
+	// Early out
+	if (FillX == 0 && FillY == 0)
+	{
+		return false;
+	}
+
+	const int32 NewSizeX = InImage.SizeX + FillX;
+	const int32 NewSizeY = InImage.SizeY + FillY;
+		
+	// Allocate room to fill out into
+	int32 BytesPerPixel = InImage.GetBytesPerPixel();
+	TArray<uint8> ResizedRawData;
+	ResizedRawData.SetNumUninitialized(NewSizeX * NewSizeY * InImage.NumSlices * BytesPerPixel);
+	
+	int32 SourceSliceSize = InImage.SizeX * InImage.SizeY * BytesPerPixel;
+	int32 DestSliceSize = NewSizeX * NewSizeY * BytesPerPixel;
+	for ( int32 SliceIndex=0; SliceIndex < InImage.NumSlices; ++SliceIndex )
+	{
+		uint8* SourceData = ((uint8*)InImage.RawData.GetData()) + SliceIndex * SourceSliceSize;
+		uint8* DestData = ((uint8*)ResizedRawData.GetData()) + SliceIndex * DestSliceSize;
+		for ( int32 Y = 0; Y < NewSizeY; ++Y )
+		{
+			for ( int32 X = 0; X < NewSizeX; ++X )
+			{
+				uint8* DestColor = DestData + Y*BytesPerPixel * NewSizeX + X*BytesPerPixel;
+				int32 SourceX = FMath::Min(X, InImage.SizeX-1);
+				int32 SourceY = FMath::Min(Y, InImage.SizeY-1);
+				uint8* SourceColor = SourceData + SourceY*BytesPerPixel * InImage.SizeX + SourceX*BytesPerPixel;
+				FMemory::Memcpy(DestColor, SourceColor, BytesPerPixel);
+			}
+		}
+	}
+
+	// Put the new image data into the OutImage
+	OutImage.Init(NewSizeX, NewSizeY, InImage.NumSlices, InImage.Format, InImage.GammaSpace);
+	FMemory::Memcpy(OutImage.RawData.GetData(), ResizedRawData.GetData(), NewSizeX * NewSizeY * OutImage.NumSlices * BytesPerPixel);
+
+	return true;
+}
+
 /**
  * Intel BC texture format handler.
  */
@@ -108,7 +161,7 @@ class FTextureFormatIntelISPCTexComp : public ITextureFormat
 
 	virtual uint16 GetVersion(FName Format) const override
 	{
-		return 0;
+		return 1;
 	}
 
 	virtual void GetSupportedFormats(TArray<FName>& OutFormats) const override
@@ -135,7 +188,8 @@ class FTextureFormatIntelISPCTexComp : public ITextureFormat
 
 		const int iWidthInBlocks	= ((InImage.SizeX + 3) & ~ 3) / 4;
 		const int iHeightInBlocks	= ((InImage.SizeY + 3) & ~ 3) / 4;
-		const int iOutputBytes		= iWidthInBlocks * iHeightInBlocks * 16;
+		const int iOutputSliceSizeBytes = iWidthInBlocks * iHeightInBlocks * 16;
+		const int iOutputBytes		= iOutputSliceSizeBytes * InImage.NumSlices;
 		OutCompressedImage.RawData.AddUninitialized(iOutputBytes);
 
 		// When we allow async tasks to execute we do so with 4 lines of the image per task
@@ -148,7 +202,16 @@ class FTextureFormatIntelISPCTexComp : public ITextureFormat
 		if ( BuildSettings.TextureFormatName == GTextureFormatNameBC6H )
 		{
 			FImage Image;
-			InImage.CopyTo(Image, ERawImageFormat::RGBA16F, EGammaSpace::Linear);
+			// Expand any dimension less than 4 to 4.  Intel compressor doens't seem to do this automatically.
+			FImage ResizedImage;
+			if (ExpandTo4PxMultiple(InImage, ResizedImage))
+			{
+				ResizedImage.CopyTo(Image, ERawImageFormat::RGBA16F, EGammaSpace::Linear);
+			}
+			else
+			{
+				InImage.CopyTo(Image, ERawImageFormat::RGBA16F, EGammaSpace::Linear);
+			}
 
 			bc6h_enc_settings settings;
 			GetProfile_bc6h_basic(&settings);
@@ -158,18 +221,19 @@ class FTextureFormatIntelISPCTexComp : public ITextureFormat
 				class FIntelCompressWorker : public FNonAbandonableTask
 				{
 				public:
-					FIntelCompressWorker(bc6h_enc_settings* pEncSettings, FImage* pInImage, FCompressedImage2D* pOutImage, int yStart, int yEnd)
+					FIntelCompressWorker(bc6h_enc_settings* pEncSettings, FImage* pInImage, FCompressedImage2D* pOutImage, int yStart, int yEnd, int iOutputSliceSizeBytes)
 						: mpEncSettings(pEncSettings)
 						, mpInImage(pInImage)
 						, mpOutImage(pOutImage)
 						, mYStart(yStart)
 						, mYEnd(yEnd)
+						, mOutputSliceSizeBytes(iOutputSliceSizeBytes)
 					{
 					}
 
 					void DoWork()
 					{
-						IntelBC6HCompressScans(mpEncSettings, mpInImage, mpOutImage, mYStart, mYEnd);
+						IntelBC6HCompressScans(mpEncSettings, mpInImage, mpOutImage, mYStart, mYEnd, mOutputSliceSizeBytes);
 					}
 
 					FORCEINLINE TStatId GetStatId() const
@@ -182,6 +246,7 @@ class FTextureFormatIntelISPCTexComp : public ITextureFormat
 					FCompressedImage2D*	mpOutImage;
 					int					mYStart;
 					int					mYEnd;
+					int					mOutputSliceSizeBytes;
 				};
 				typedef FAsyncTask<FIntelCompressWorker> FIntelCompressTask;
 
@@ -190,12 +255,10 @@ class FTextureFormatIntelISPCTexComp : public ITextureFormat
 				CompressionTasks.Reserve(iNumTasks);
 				for ( int iTask=0; iTask < iNumTasks; ++iTask )
 				{
-					auto* AsyncTask = new(CompressionTasks) FIntelCompressTask(&settings, &Image, &OutCompressedImage, iTask * iScansPerTask, (iTask + 1) * iScansPerTask);
+					auto* AsyncTask = new(CompressionTasks) FIntelCompressTask(&settings, &Image, &OutCompressedImage, iTask * iScansPerTask, (iTask + 1) * iScansPerTask, iOutputSliceSizeBytes);
 					AsyncTask->StartBackgroundTask();
 				}
-
-				IntelBC6HCompressScans(&settings, &Image, &OutCompressedImage, iScansPerTask * iNumTasks, InImage.SizeY);
-
+				IntelBC6HCompressScans(&settings, &Image, &OutCompressedImage, iScansPerTask * iNumTasks, Image.SizeY, iOutputSliceSizeBytes);
 				// Wait completion
 				for (int32 TaskIndex = 0; TaskIndex < CompressionTasks.Num(); ++TaskIndex)
 				{
@@ -204,7 +267,7 @@ class FTextureFormatIntelISPCTexComp : public ITextureFormat
 			}
 			else
 			{
-				IntelBC6HCompressScans(&settings, &Image, &OutCompressedImage, 0, InImage.SizeY);
+				IntelBC6HCompressScans(&settings, &Image, &OutCompressedImage, 0, Image.SizeY, iOutputSliceSizeBytes);
 			}
 
 			CompressedPixelFormat = PF_BC6H;
@@ -212,8 +275,18 @@ class FTextureFormatIntelISPCTexComp : public ITextureFormat
 		}
 		else if ( BuildSettings.TextureFormatName == GTextureFormatNameBC7 )
 		{
+			check(InImage.NumSlices == 1);
 			FImage Image;
-			InImage.CopyTo(Image, ERawImageFormat::BGRA8, BuildSettings.GetGammaSpace());
+			// Expand any dimension less than 4 to 4.  Intel compressor doesn't seem to do this automatically.
+			FImage ResizedImage;
+			if (ExpandTo4PxMultiple(InImage, ResizedImage))
+			{
+				ResizedImage.CopyTo(Image, ERawImageFormat::BGRA8, BuildSettings.GetGammaSpace());
+			}
+			else
+			{
+				InImage.CopyTo(Image, ERawImageFormat::BGRA8, BuildSettings.GetGammaSpace());
+			}
 
 			bc7_enc_settings settings;
 			if ( bImageHasAlphaChannel )
@@ -266,7 +339,7 @@ class FTextureFormatIntelISPCTexComp : public ITextureFormat
 					AsyncTask->StartBackgroundTask();
 				}
 
-				IntelBC7CompressScans(&settings, &Image, &OutCompressedImage, iScansPerTask * iNumTasks, InImage.SizeY);
+				IntelBC7CompressScans(&settings, &Image, &OutCompressedImage, iScansPerTask * iNumTasks, Image.SizeY);
 
 				// Wait completion
 				for (int32 TaskIndex = 0; TaskIndex < CompressionTasks.Num(); ++TaskIndex)
@@ -276,7 +349,7 @@ class FTextureFormatIntelISPCTexComp : public ITextureFormat
 			}
 			else
 			{
-				IntelBC7CompressScans(&settings, &Image, &OutCompressedImage, 0, InImage.SizeY);
+				IntelBC7CompressScans(&settings, &Image, &OutCompressedImage, 0, Image.SizeY);
 			}
 
 			CompressedPixelFormat = PF_BC7;
diff --git a/Engine/Source/Editor/HardwareTargeting/Private/HardwareTargetingModule.cpp b/Engine/Source/Editor/HardwareTargeting/Private/HardwareTargetingModule.cpp
--- a/Engine/Source/Editor/HardwareTargeting/Private/HardwareTargetingModule.cpp
+++ b/Engine/Source/Editor/HardwareTargeting/Private/HardwareTargetingModule.cpp
@@ -97,22 +97,23 @@ FText FMetaSettingGatherer::ValueToString(bool Value)
 }
 
 template <>
-FText FMetaSettingGatherer::ValueToString(EAntiAliasingMethodUI::Type Value)
+FText FMetaSettingGatherer::ValueToString(EAntiAliasingMethod Value)
 {
 	switch (Value)
 	{
-	case EAntiAliasingMethodUI::AAM_None:
+	case AAM_None:
 		return LOCTEXT("AA_None", "None");
-	case EAntiAliasingMethodUI::AAM_FXAA:
+	case AAM_FXAA:
 		return LOCTEXT("AA_FXAA", "FXAA");
-	case EAntiAliasingMethodUI::AAM_TemporalAA:
+	case AAM_TemporalAA:
 		return LOCTEXT("AA_TemporalAA", "Temporal AA");
+	case AAM_MSAA:
+		return LOCTEXT("AA_MSAA", "MSAA");
 	default:
 		return FText::AsNumber((int32)Value);
 	}
 }
 
-
 #define UE_META_SETTING_ENTRY(Builder, Class, PropertyName, TargetValue) \
 { \
 	Class* SettingsObject = GetMutableDefault<Class>(); \
@@ -228,7 +229,7 @@ void FHardwareTargetingModule::GatherSettings(FMetaSettingGatherer& Builder)
 		// DOF and AA work on mobile but are expensive, keeping them off by default
 		//@TODO: DOF setting doesn't exist yet
 		// UE_META_SETTING_ENTRY(Builder, URendererSettings, bDefaultFeatureDepthOfField, bHighEndPC);
-		UE_META_SETTING_ENTRY(Builder, URendererSettings, DefaultFeatureAntiAliasing, bHighEndPC ? EAntiAliasingMethodUI::AAM_TemporalAA : EAntiAliasingMethodUI::AAM_None);
+		UE_META_SETTING_ENTRY(Builder, URendererSettings, DefaultFeatureAntiAliasing, bHighEndPC ? AAM_TemporalAA : AAM_None);
 	}
 
 	{
diff --git a/Engine/Source/Runtime/Core/Private/Modules/ModuleManager.cpp b/Engine/Source/Runtime/Core/Private/Modules/ModuleManager.cpp
--- a/Engine/Source/Runtime/Core/Private/Modules/ModuleManager.cpp
+++ b/Engine/Source/Runtime/Core/Private/Modules/ModuleManager.cpp
@@ -331,7 +331,7 @@ TSharedPtr<IModuleInterface> FModuleManager::LoadModule( const FName InModuleNam
 TSharedPtr<IModuleInterface> FModuleManager::LoadModuleChecked( const FName InModuleName, const bool bWasReloaded )
 {
 	TSharedPtr<IModuleInterface> Module = LoadModule(InModuleName, bWasReloaded);
-	check(Module.IsValid());
+	checkf(Module.IsValid(), *InModuleName.ToString());
 
 	return Module;
 }
diff --git a/Engine/Source/Runtime/Core/Public/UObject/RenderingObjectVersion.h b/Engine/Source/Runtime/Core/Public/UObject/RenderingObjectVersion.h
--- a/Engine/Source/Runtime/Core/Public/UObject/RenderingObjectVersion.h
+++ b/Engine/Source/Runtime/Core/Public/UObject/RenderingObjectVersion.h
@@ -12,15 +12,17 @@ struct CORE_API FRenderingObjectVersion
 		// Added support for 3 band SH in the ILC
 		IndirectLightingCache3BandSupport,
 
-		// Allows specifying resolution for reflecture capture probes
+		// Allows specifying resolution for reflection capture probes
 		CustomReflectionCaptureResolutionSupport,
 
-		// Allows specifying resolution for reflecture capture probes
 		RemovedTextureStreamingLevelData,
 
 		// translucency is now a property which matters for materials with the decal domain
 		IntroducedMeshDecals,
 
+		// Reflection captures are no longer prenormalized
+		ReflectionCapturesStoreAverageBrightness,
+
 		// -----<new versions can be added above this line>-------------------------------------------------
 		VersionPlusOne,
 		LatestVersion = VersionPlusOne - 1
diff --git a/Engine/Source/Runtime/Engine/Classes/Components/ReflectionCaptureComponent.h b/Engine/Source/Runtime/Engine/Classes/Components/ReflectionCaptureComponent.h
--- a/Engine/Source/Runtime/Engine/Classes/Components/ReflectionCaptureComponent.h
+++ b/Engine/Source/Runtime/Engine/Classes/Components/ReflectionCaptureComponent.h
@@ -144,6 +144,14 @@ class UReflectionCaptureComponent : public USceneComponent
 		return FullHDRData;
 	}
 
+	inline float GetAverageBrightness() const 
+	{ 
+		return AverageBrightness; 
+	}
+
+	inline float* GetAverageBrightnessPtr() { return &AverageBrightness; }
+	inline const float* GetAverageBrightnessPtr() const { return &AverageBrightness; }
+
 	ENGINE_API static int32 GetReflectionCaptureSize_GameThread();
 	ENGINE_API static int32 GetReflectionCaptureSize_RenderThread();
 
@@ -181,6 +189,9 @@ private:
 	UPROPERTY()
 	FGuid StateId;
 
+	/** Average brightness of the captured data, read back to the CPU after the capture. */
+	float AverageBrightness;
+
 	/**
 	 * The full HDR capture data to use for rendering.
 	 * This will be loaded from inlined data.
diff --git a/Engine/Source/Runtime/Engine/Classes/Components/SkyLightComponent.h b/Engine/Source/Runtime/Engine/Classes/Components/SkyLightComponent.h
--- a/Engine/Source/Runtime/Engine/Classes/Components/SkyLightComponent.h
+++ b/Engine/Source/Runtime/Engine/Classes/Components/SkyLightComponent.h
@@ -205,7 +205,7 @@ class ENGINE_API USkyLightComponent : public ULightComponentBase
 	/** Indicates that the capture needs to recapture the scene, adds it to the recapture queue. */
 	void SetCaptureIsDirty();
 	void SetBlendDestinationCaptureIsDirty();
-	void SanatizeCubemapSize();
+	void SanitizeCubemapSize();
 
 	/** 
 	 * Recaptures the scene for the skylight. 
@@ -229,6 +229,7 @@ protected:
 
 	TRefCountPtr<FSkyTextureCubeResource> ProcessedSkyTexture;
 	FSHVectorRGB3 IrradianceEnvironmentMap;
+	float AverageBrightness;
 
 	/** If 0, no blend is present.  If > 0, BlendDestinationProcessedSkyTexture and BlendDestinationIrradianceEnvironmentMap must be generated and used for rendering. */
 	float BlendFraction;
@@ -237,6 +238,7 @@ protected:
 	class UTextureCube* BlendDestinationCubemap;
 	TRefCountPtr<FSkyTextureCubeResource> BlendDestinationProcessedSkyTexture;
 	FSHVectorRGB3 BlendDestinationIrradianceEnvironmentMap;
+	float BlendDestinationAverageBrightness;
 
 	/** Tracks when the rendering thread has completed its writes to IrradianceEnvironmentMap. */
 	FRenderCommandFence IrradianceMapFence;
diff --git a/Engine/Source/Runtime/Engine/Classes/Engine/RendererSettings.h b/Engine/Source/Runtime/Engine/Classes/Engine/RendererSettings.h
--- a/Engine/Source/Runtime/Engine/Classes/Engine/RendererSettings.h
+++ b/Engine/Source/Runtime/Engine/Classes/Engine/RendererSettings.h
@@ -70,20 +70,6 @@ namespace EEarlyZPass
 	};
 }
 
-
-/** used by FPostProcessSettings Anti-aliasing */
-UENUM()
-namespace EAntiAliasingMethodUI
-{
-	enum Type
-	{
-		AAM_None UMETA(DisplayName = "None"),
-		AAM_FXAA UMETA(DisplayName = "FXAA"),
-		AAM_TemporalAA UMETA(DisplayName = "TemporalAA"),
-		AAM_MAX,
-	};
-}
-
 /** used by FPostProcessSettings AutoExposure*/
 UENUM()
 namespace EAutoExposureMethodUI
@@ -187,6 +173,12 @@ class ENGINE_API URendererSettings : public UDeveloperSettings
 		ConfigRestartRequired=true))
 	uint32 bForwardShading:1;
 
+	UPROPERTY(config, EditAnywhere, Category=ForwardShading, meta=(
+		ConsoleVariable="r.VertexFoggingForOpaque",
+		ToolTip="Causes opaque materials to use per-vertex fogging, which costs less and integrates properly with MSAA.  Only supported with forward shading. Changing this setting requires restarting the editor.",
+		ConfigRestartRequired=true))
+	uint32 bVertexFoggingForOpaque:1;
+
 	UPROPERTY(config, EditAnywhere, Category=Lighting, meta=(
 		ConsoleVariable="r.AllowStaticLighting",
 		ToolTip="Whether to allow any static lighting to be generated and used, like lightmaps and shadowmaps. Games that only use dynamic lighting should set this to 0 to save some static lighting overhead. Changing this setting requires restarting the editor.",
@@ -235,45 +227,45 @@ class ENGINE_API URendererSettings : public UDeveloperSettings
 		ToolTip="Whether the custom depth pass for tagging primitives for postprocessing passes is enabled. Enabling it on demand can save memory but may cause a hitch the first time the feature is used."))
 	TEnumAsByte<ECustomDepthStencil::Type> CustomDepthStencil;
 
-	UPROPERTY(config, EditAnywhere, Category = DefaultPostprocessingSettings, meta = (
+	UPROPERTY(config, EditAnywhere, Category = DefaultSettings, meta = (
 		ConsoleVariable = "r.DefaultFeature.Bloom", DisplayName = "Bloom",
 		ToolTip = "Whether the default for Bloom is enabled or not (postprocess volume/camera/game setting can still override and enable or disable it independently)"))
 	uint32 bDefaultFeatureBloom : 1;
 
-	UPROPERTY(config, EditAnywhere, Category = DefaultPostprocessingSettings, meta = (
+	UPROPERTY(config, EditAnywhere, Category = DefaultSettings, meta = (
 		ConsoleVariable = "r.DefaultFeature.AmbientOcclusion", DisplayName = "Ambient Occlusion",
 		ToolTip = "Whether the default for AmbientOcclusion is enabled or not (postprocess volume/camera/game setting can still override and enable or disable it independently)"))
 	uint32 bDefaultFeatureAmbientOcclusion : 1;
 
-	UPROPERTY(config, EditAnywhere, Category = DefaultPostprocessingSettings, meta = (
+	UPROPERTY(config, EditAnywhere, Category = DefaultSettings, meta = (
 		ConsoleVariable = "r.DefaultFeature.AmbientOcclusionStaticFraction", DisplayName = "Ambient Occlusion Static Fraction (AO for baked lighting)",
 		ToolTip = "Whether the default for AmbientOcclusionStaticFraction is enabled or not (only useful for baked lighting and if AO is on, allows to have SSAO affect baked lighting as well, costs performance, postprocess volume/camera/game setting can still override and enable or disable it independently)"))
 	uint32 bDefaultFeatureAmbientOcclusionStaticFraction : 1;
 
-	UPROPERTY(config, EditAnywhere, Category = DefaultPostprocessingSettings, meta = (
+	UPROPERTY(config, EditAnywhere, Category = DefaultSettings, meta = (
 		ConsoleVariable = "r.DefaultFeature.AutoExposure", DisplayName = "Auto Exposure",
 		ToolTip = "Whether the default for AutoExposure is enabled or not (postprocess volume/camera/game setting can still override and enable or disable it independently)"))
 	uint32 bDefaultFeatureAutoExposure : 1;
 	
-	UPROPERTY(config, EditAnywhere, Category = DefaultPostprocessingSettings, meta = (
+	UPROPERTY(config, EditAnywhere, Category = DefaultSettings, meta = (
 		ConsoleVariable = "r.DefaultFeature.AutoExposure.Method", DisplayName = "Auto Exposure",
 		ToolTip = "The default method for AutoExposure(postprocess volume/camera/game setting can still override and enable or disable it independently)"))
 	TEnumAsByte<EAutoExposureMethodUI::Type> DefaultFeatureAutoExposure; 
 
-	UPROPERTY(config, EditAnywhere, Category = DefaultPostprocessingSettings, meta = (
+	UPROPERTY(config, EditAnywhere, Category = DefaultSettings, meta = (
 		ConsoleVariable = "r.DefaultFeature.MotionBlur", DisplayName = "Motion Blur",
 		ToolTip = "Whether the default for MotionBlur is enabled or not (postprocess volume/camera/game setting can still override and enable or disable it independently)"))
 	uint32 bDefaultFeatureMotionBlur : 1;
 
-	UPROPERTY(config, EditAnywhere, Category = DefaultPostprocessingSettings, meta = (
+	UPROPERTY(config, EditAnywhere, Category = DefaultSettings, meta = (
 		ConsoleVariable = "r.DefaultFeature.LensFlare", DisplayName = "Lens Flares (Image based)",
 		ToolTip = "Whether the default for LensFlare is enabled or not (postprocess volume/camera/game setting can still override and enable or disable it independently)"))
 	uint32 bDefaultFeatureLensFlare : 1;
 
-	UPROPERTY(config, EditAnywhere, Category = DefaultPostprocessingSettings, meta = (
+	UPROPERTY(config, EditAnywhere, Category = DefaultSettings, meta = (
 		ConsoleVariable = "r.DefaultFeature.AntiAliasing", DisplayName = "Anti-Aliasing Method",
 		ToolTip = "What anti-aliasing mode is used by default (postprocess volume/camera/game setting can still override and enable or disable it independently)"))
-	TEnumAsByte<EAntiAliasingMethodUI::Type> DefaultFeatureAntiAliasing;
+	TEnumAsByte<EAntiAliasingMethod> DefaultFeatureAntiAliasing;
 
 	UPROPERTY(config, EditAnywhere, Category=Optimizations, meta=(
 		ConsoleVariable="r.StencilForLODDither",DisplayName="Use Stencil for LOD Dither Fading",
@@ -293,7 +285,8 @@ class ENGINE_API URendererSettings : public UDeveloperSettings
 
 	UPROPERTY(config, EditAnywhere, Category=Lighting, meta=(
 		ConsoleVariable="r.DBuffer",DisplayName="DBuffer Decals",
-		ToolTip="Experimental decal feature (see r.DBuffer, ideally combined with 'Movables in early Z-pass' and 'Early Z-pass')"))
+		ToolTip="Whether to accumulate decal properties to a buffer before the base pass.  DBuffer decals correctly affect lightmap and sky lighting, unlike regular deferred decals.  DBuffer enabled forces a full prepass.  Changing this setting requires restarting the editor.",
+		ConfigRestartRequired=true))
 	uint32 bDBuffer:1;
 
 	UPROPERTY(config, EditAnywhere, Category=Optimizations, meta=(
diff --git a/Engine/Source/Runtime/Engine/Classes/Engine/Scene.h b/Engine/Source/Runtime/Engine/Classes/Engine/Scene.h
--- a/Engine/Source/Runtime/Engine/Classes/Engine/Scene.h
+++ b/Engine/Source/Runtime/Engine/Classes/Engine/Scene.h
@@ -18,13 +18,15 @@ enum EDepthOfFieldMethod
 	DOFM_MAX,
 };
 
-/** Used by FPostProcessSettings Anti-aliasings */
+/** Used by rendering project settings. */
 UENUM()
 enum EAntiAliasingMethod
 {
 	AAM_None UMETA(DisplayName="None"),
 	AAM_FXAA UMETA(DisplayName="FXAA"),
 	AAM_TemporalAA UMETA(DisplayName="TemporalAA"),
+	/** Only supported with forward shading.  MSAA sample count is controlled by r.MSAACount. */
+	AAM_MSAA UMETA(DisplayName="MSAA"),
 	AAM_MAX,
 };
 
@@ -1016,10 +1018,6 @@ struct FPostProcessSettings
 	UPROPERTY(interp, BlueprintReadWrite, Category=Misc, AdvancedDisplay, meta=(ClampMin = "0.0", ClampMax = "400.0", editcondition = "bOverride_ScreenPercentage"))
 	float ScreenPercentage;
 
-	/** TemporalAA, FXAA, ... */
-	UPROPERTY(EditAnywhere, BlueprintReadWrite, AdvancedDisplay, Category=Misc, meta=(editcondition = "bOverride_AntiAliasingMethod", DisplayName = "AA Method"))
-	TEnumAsByte<enum EAntiAliasingMethod> AntiAliasingMethod;
-
 	/** Enable/Fade/disable the Screen Space Reflection feature, in percent, avoid numbers between 0 and 1 fo consistency */
 	UPROPERTY(interp, BlueprintReadWrite, Category=ScreenSpaceReflections, meta=(ClampMin = "0.0", ClampMax = "100.0", editcondition = "bOverride_ScreenSpaceReflectionIntensity", DisplayName = "Intensity"))
 	float ScreenSpaceReflectionIntensity;
@@ -1263,7 +1261,6 @@ struct FPostProcessSettings
 		MotionBlurMax = 5.0f;
 		MotionBlurPerObjectSize = 0.5f;
 		ScreenPercentage = 100.0f;
-		AntiAliasingMethod = AAM_TemporalAA;
 		ScreenSpaceReflectionIntensity = 100.0f;
 		ScreenSpaceReflectionQuality = 50.0f;
 		ScreenSpaceReflectionMaxRoughness = 0.6f;
diff --git a/Engine/Source/Runtime/Engine/Classes/Materials/Material.h b/Engine/Source/Runtime/Engine/Classes/Materials/Material.h
--- a/Engine/Source/Runtime/Engine/Classes/Materials/Material.h
+++ b/Engine/Source/Runtime/Engine/Classes/Materials/Material.h
@@ -651,6 +651,10 @@ public:
 	UPROPERTY(EditAnywhere, BlueprintReadOnly, Category = ForwardShading, meta = (DisplayName = "Planar Reflections"))
 	uint32 bUsePlanarForwardReflections : 1;
 
+	/* Reduce roughness based on screen space normal changes. */
+	UPROPERTY(EditAnywhere, BlueprintReadOnly, Category=ForwardShading)
+	uint32 bNormalCurvatureToRoughness : 1;
+
 	/** The type of tessellation to apply to this object.  Note D3D11 required for anything except MTM_NoTessellation. */
 	UPROPERTY(EditAnywhere, BlueprintReadOnly, Category=Tessellation)
 	TEnumAsByte<enum EMaterialTessellationMode> D3D11TessellationMode;
diff --git a/Engine/Source/Runtime/Engine/Classes/Materials/MaterialExpressionNoise.h b/Engine/Source/Runtime/Engine/Classes/Materials/MaterialExpressionNoise.h
--- a/Engine/Source/Runtime/Engine/Classes/Materials/MaterialExpressionNoise.h
+++ b/Engine/Source/Runtime/Engine/Classes/Materials/MaterialExpressionNoise.h
@@ -22,7 +22,7 @@ enum ENoiseFunction
 	 */
 	NOISEFUNCTION_GradientTex UMETA(DisplayName="Gradient - Texture Based"),
 
-	/** High quality for direct use, BAD for bumps
+	/** High quality for direct use, BAD for bumps; doesn't work on Mobile
 	 * ~16 instructions per level, 1 texture lookup
 	 * Always tiles with a repeat of 16, "Tiling" mode is not an option for Fast Gradient noise
 	 */
diff --git a/Engine/Source/Runtime/Engine/Private/Components/ReflectionCaptureComponent.cpp b/Engine/Source/Runtime/Engine/Private/Components/ReflectionCaptureComponent.cpp
--- a/Engine/Source/Runtime/Engine/Private/Components/ReflectionCaptureComponent.cpp
+++ b/Engine/Source/Runtime/Engine/Private/Components/ReflectionCaptureComponent.cpp
@@ -289,7 +289,7 @@ APlaneReflectionCapture::APlaneReflectionCapture(const FObjectInitializer& Objec
 // Generate a new guid to force a recapture of all reflection data
 // Note: changing this will cause saved capture data in maps to be discarded
 // A resave of those maps will be required to guarantee valid reflections when cooking for ES2
-FGuid ReflectionCaptureDDCVer(0x0c669396, 0x9cb849ae, 0x9f4120ff, 0x5812f4d2);
+FGuid ReflectionCaptureDDCVer(0x0c669396, 0x9cb849ae, 0x9f4120ff, 0x5812f4d3);
 
 FReflectionCaptureFullHDR::~FReflectionCaptureFullHDR()
 {
@@ -812,6 +812,7 @@ UReflectionCaptureComponent::UReflectionCaptureComponent(const FObjectInitialize
 
 	bCaptureDirty = false;
 	bDerivedDataDirty = false;
+	AverageBrightness = 1.0f;
 }
 
 void UReflectionCaptureComponent::CreateRenderState_Concurrent()
@@ -871,6 +872,7 @@ void UReflectionCaptureComponent::SerializeSourceData(FArchive& Ar)
 		if (Ar.IsSaving())
 		{
 			Ar << ReflectionCaptureDDCVer;
+			Ar << AverageBrightness;
 
 			int32 StartOffset = Ar.Tell();
 			Ar << StartOffset;
@@ -894,6 +896,11 @@ void UReflectionCaptureComponent::SerializeSourceData(FArchive& Ar)
 			FGuid SavedVersion;
 			Ar << SavedVersion;
 
+			if (Ar.CustomVer(FRenderingObjectVersion::GUID) >= FRenderingObjectVersion::ReflectionCapturesStoreAverageBrightness)
+			{
+				Ar << AverageBrightness;
+			}
+
 			int32 EndOffset = 0;
 			Ar << EndOffset;
 
@@ -959,6 +966,8 @@ void UReflectionCaptureComponent::Serialize(FArchive& Ar)
 		// Saving for cooking path
 		if (Ar.IsCooking())
 		{
+			Ar << AverageBrightness;
+
 			// Get all the reflection capture formats that the target platform wants
 			TArray<FName> Formats;
 			Ar.CookingTarget()->GetReflectionCaptureFormats(Formats);
@@ -1010,7 +1019,7 @@ void UReflectionCaptureComponent::Serialize(FArchive& Ar)
 					else if (!IsTemplate())
 					{
 						// Temporary warning until the cooker can do scene captures itself in the case of missing DDC
-						UE_LOG(LogMaterial, Warning, TEXT("Reflection capture requires encoded HDR data but none was found in the DDC!  This reflection will be black.  Fix by loading the map in the editor once.  %s."), *GetFullName());
+						UE_LOG(LogMaterial, Warning, TEXT("Reflection capture requires encoded HDR data but none was found in the DDC!  This reflection will be black.  Fix by resaving the map in the editor.  %s."), *GetFullName());
 					}
 				}
 			}
@@ -1018,6 +1027,7 @@ void UReflectionCaptureComponent::Serialize(FArchive& Ar)
 		else
 		{
 			// Loading cooked data path
+			Ar << AverageBrightness;
 
 			int32 NumFormats = 0;
 			Ar << NumFormats;
@@ -1036,15 +1046,7 @@ void UReflectionCaptureComponent::Serialize(FArchive& Ar)
 					{
 						FullHDRData = new FReflectionCaptureFullHDR();
 
-						if (Ar.CustomVer(FRenderingObjectVersion::GUID) >= FRenderingObjectVersion::CustomReflectionCaptureResolutionSupport)
-						{
-							Ar << FullHDRData->CubemapSize;
-						}
-						else
-						{
-							FullHDRData->CubemapSize = 128;
-						}
-
+						Ar << FullHDRData->CubemapSize;
 						Ar << FullHDRData->CompressedCapturedData;
 					}
 					else 
@@ -1682,6 +1684,22 @@ FReflectionCaptureProxy::FReflectionCaptureProxy(const UReflectionCaptureCompone
 	InfluenceRadius = InComponent->GetInfluenceBoundingRadius();
 	Brightness = InComponent->Brightness;
 	Guid = GetTypeHash( Component->GetPathName() );
+	AverageBrightness = 1.0f;
+
+	ENQUEUE_UNIQUE_RENDER_COMMAND_TWOPARAMETER(
+		FInitReflectionProxy,
+		const float*,AverageBrightness,InComponent->GetAverageBrightnessPtr(),
+		FReflectionCaptureProxy*,ReflectionCaptureProxy,this,
+	{
+		// Only access AverageBrightness on the RT, even though they belong to the UReflectionCaptureComponent, 
+		// Because FScene::UpdateReflectionCaptureContents does not block the RT so the writes could still be in flight
+		ReflectionCaptureProxy->InitializeAverageBrightness(*AverageBrightness);
+	});
+}
+
+void FReflectionCaptureProxy::InitializeAverageBrightness(const float& InAverageBrightness)
+{
+	AverageBrightness = InAverageBrightness;
 }
 
 void FReflectionCaptureProxy::SetTransform(const FMatrix& InTransform)
diff --git a/Engine/Source/Runtime/Engine/Private/Components/SkyLightComponent.cpp b/Engine/Source/Runtime/Engine/Private/Components/SkyLightComponent.cpp
--- a/Engine/Source/Runtime/Engine/Private/Components/SkyLightComponent.cpp
+++ b/Engine/Source/Runtime/Engine/Private/Components/SkyLightComponent.cpp
@@ -70,7 +70,12 @@ void UWorld::UpdateAllSkyCaptures()
 	USkyLightComponent::UpdateSkyCaptureContents(this);
 }
 
-void FSkyLightSceneProxy::Initialize(float InBlendFraction, const FSHVectorRGB3* InIrradianceEnvironmentMap, const FSHVectorRGB3* BlendDestinationIrradianceEnvironmentMap)
+void FSkyLightSceneProxy::Initialize(
+	float InBlendFraction, 
+	const FSHVectorRGB3* InIrradianceEnvironmentMap, 
+	const FSHVectorRGB3* BlendDestinationIrradianceEnvironmentMap,
+	const float* InAverageBrightness,
+	const float* BlendDestinationAverageBrightness)
 {
 	BlendFraction = FMath::Clamp(InBlendFraction, 0.0f, 1.0f);
 
@@ -79,17 +84,20 @@ void FSkyLightSceneProxy::Initialize(float InBlendFraction, const FSHVectorRGB3*
 		if (BlendFraction < 1)
 		{
 			IrradianceEnvironmentMap = (*InIrradianceEnvironmentMap) * (1 - BlendFraction) + (*BlendDestinationIrradianceEnvironmentMap) * BlendFraction;
+			AverageBrightness = *InAverageBrightness * (1 - BlendFraction) + (*BlendDestinationAverageBrightness) * BlendFraction;
 		}
 		else
 		{
 			// Blend is full destination, treat as source to avoid blend overhead in shaders
 			IrradianceEnvironmentMap = *BlendDestinationIrradianceEnvironmentMap;
+			AverageBrightness = *BlendDestinationAverageBrightness;
 		}
 	}
 	else
 	{
 		// Blend is full source
 		IrradianceEnvironmentMap = *InIrradianceEnvironmentMap;
+		AverageBrightness = *InAverageBrightness;
 		BlendFraction = 0;
 	}
 }
@@ -110,16 +118,18 @@ FSkyLightSceneProxy::FSkyLightSceneProxy(const USkyLightComponent* InLightCompon
 	, MinOcclusion(InLightComponent->MinOcclusion)
 	, OcclusionTint(InLightComponent->OcclusionTint)
 {
-	ENQUEUE_UNIQUE_RENDER_COMMAND_FOURPARAMETER(
+	ENQUEUE_UNIQUE_RENDER_COMMAND_SIXPARAMETER(
 		FInitSkyProxy,
 		const FSHVectorRGB3*,InIrradianceEnvironmentMap,&InLightComponent->IrradianceEnvironmentMap,
 		const FSHVectorRGB3*,BlendDestinationIrradianceEnvironmentMap,&InLightComponent->BlendDestinationIrradianceEnvironmentMap,
+		const float*,InAverageBrightness,&InLightComponent->AverageBrightness,
+		const float*,BlendDestinationAverageBrightness,&InLightComponent->BlendDestinationAverageBrightness,
 		float,InBlendFraction,InLightComponent->BlendFraction,
 		FSkyLightSceneProxy*,LightSceneProxy,this,
 	{
 		// Only access the irradiance maps on the RT, even though they belong to the USkyLightComponent, 
 		// Because FScene::UpdateSkyCaptureContents does not block the RT so the writes could still be in flight
-		LightSceneProxy->Initialize(InBlendFraction, InIrradianceEnvironmentMap, BlendDestinationIrradianceEnvironmentMap);
+		LightSceneProxy->Initialize(InBlendFraction, InIrradianceEnvironmentMap, BlendDestinationIrradianceEnvironmentMap, InAverageBrightness, BlendDestinationAverageBrightness);
 	});
 }
 
@@ -150,6 +160,8 @@ USkyLightComponent::USkyLightComponent(const FObjectInitializer& ObjectInitializ
 	OcclusionTint = FColor::Black;
 	CubemapResolution = 128;
 	LowerHemisphereColor = FLinearColor::Black;
+	AverageBrightness = 1.0f;
+	BlendDestinationAverageBrightness = 1.0f;
 }
 
 FSkyLightSceneProxy* USkyLightComponent::CreateSceneProxy() const
@@ -173,7 +185,7 @@ void USkyLightComponent::SetCaptureIsDirty()
 	}
 }
 
-void USkyLightComponent::SanatizeCubemapSize()
+void USkyLightComponent::SanitizeCubemapSize()
 {
 	static const int32 MaxCubemapResolution = 1024;
 	static const int32 MinCubemapResolution = 64;
@@ -239,7 +251,7 @@ void USkyLightComponent::PostLoad()
 {
 	Super::PostLoad();
 
-	SanatizeCubemapSize();
+	SanitizeCubemapSize();
 
 	// All components are queued for update on creation by default, remove if needed
 	if (!bVisible || HasAnyFlags(RF_ClassDefaultObject))
@@ -316,7 +328,7 @@ void USkyLightComponent::PostEditChangeProperty(FPropertyChangedEvent& PropertyC
 {
 	Super::PostEditChangeProperty(PropertyChangedEvent);
 
-	SanatizeCubemapSize();
+	SanitizeCubemapSize();
 	SetCaptureIsDirty();
 }
 
@@ -431,6 +443,7 @@ public:
 	// This has to be refcounted to keep it alive during the handoff without doing a deep copy
 	TRefCountPtr<FSkyTextureCubeResource> ProcessedSkyTexture;
 	FSHVectorRGB3 IrradianceEnvironmentMap;
+	float AverageBrightness;
 };
 
 FActorComponentInstanceData* USkyLightComponent::GetComponentInstanceData() const
@@ -443,6 +456,7 @@ FActorComponentInstanceData* USkyLightComponent::GetComponentInstanceData() cons
 	// Block until the rendering thread has completed its writes from a previous capture
 	IrradianceMapFence.Wait();
 	InstanceData->IrradianceEnvironmentMap = IrradianceEnvironmentMap;
+	InstanceData->AverageBrightness = AverageBrightness;
 
 	return InstanceData;
 }
@@ -455,6 +469,7 @@ void USkyLightComponent::ApplyComponentInstanceData(FPrecomputedSkyLightInstance
 	bPrecomputedLightingIsValid = LightMapData->bPrecomputedLightingIsValid;
 	ProcessedSkyTexture = LightMapData->ProcessedSkyTexture;
 	IrradianceEnvironmentMap = LightMapData->IrradianceEnvironmentMap;
+	AverageBrightness = LightMapData->AverageBrightness;
 
 	if (ProcessedSkyTexture && bSavedConstructionScriptValuesValid)
 	{
@@ -495,7 +510,7 @@ void USkyLightComponent::UpdateSkyCaptureContentsArray(UWorld* WorldToUpdate, TA
 						CaptureComponent->MarkRenderStateDirty();
 					}
 
-					WorldToUpdate->Scene->UpdateSkyCaptureContents(CaptureComponent, false, CaptureComponent->Cubemap, CaptureComponent->ProcessedSkyTexture, CaptureComponent->IrradianceEnvironmentMap);
+					WorldToUpdate->Scene->UpdateSkyCaptureContents(CaptureComponent, false, CaptureComponent->Cubemap, CaptureComponent->ProcessedSkyTexture, CaptureComponent->AverageBrightness, CaptureComponent->IrradianceEnvironmentMap);
 				}
 				else
 				{
@@ -510,7 +525,7 @@ void USkyLightComponent::UpdateSkyCaptureContentsArray(UWorld* WorldToUpdate, TA
 						CaptureComponent->MarkRenderStateDirty(); 
 					}
 
-					WorldToUpdate->Scene->UpdateSkyCaptureContents(CaptureComponent, false, CaptureComponent->BlendDestinationCubemap, CaptureComponent->BlendDestinationProcessedSkyTexture, CaptureComponent->BlendDestinationIrradianceEnvironmentMap);
+					WorldToUpdate->Scene->UpdateSkyCaptureContents(CaptureComponent, false, CaptureComponent->BlendDestinationCubemap, CaptureComponent->BlendDestinationProcessedSkyTexture, CaptureComponent->BlendDestinationAverageBrightness, CaptureComponent->BlendDestinationIrradianceEnvironmentMap);
 				}
 
 				CaptureComponent->IrradianceMapFence.BeginFence();
@@ -540,9 +555,10 @@ void USkyLightComponent::CaptureEmissiveIrradianceEnvironmentMap(FSHVectorRGB3&
 
 	if (GetScene() && (SourceType != SLS_SpecifiedCubemap || Cubemap))
 	{
+		float UnusedAverageBrightness = 1.0f;
 		// Capture emissive scene lighting only for the lighting build
 		// This is necessary to avoid a feedback loop with the last lighting build results
-		GetScene()->UpdateSkyCaptureContents(this, true, Cubemap, NULL, OutIrradianceMap);
+		GetScene()->UpdateSkyCaptureContents(this, true, Cubemap, NULL, UnusedAverageBrightness, OutIrradianceMap);
 		// Wait until writes to OutIrradianceMap have completed
 		FlushRenderingCommands();
 	}
@@ -623,16 +639,18 @@ void USkyLightComponent::SetCubemapBlend(UTextureCube* SourceCubemap, UTextureCu
 
 			if (SceneProxy)
 			{
-				ENQUEUE_UNIQUE_RENDER_COMMAND_FOURPARAMETER(
+				ENQUEUE_UNIQUE_RENDER_COMMAND_SIXPARAMETER(
 					FUpdateSkyProxy,
 					const FSHVectorRGB3*,InIrradianceEnvironmentMap,&IrradianceEnvironmentMap,
 					const FSHVectorRGB3*,BlendDestinationIrradianceEnvironmentMap,&BlendDestinationIrradianceEnvironmentMap,
+					const float*,InAverageBrightness,&AverageBrightness,
+					const float*,BlendDestinationAverageBrightness,&BlendDestinationAverageBrightness,
 					float,InBlendFraction,BlendFraction,
 					FSkyLightSceneProxy*,LightSceneProxy,SceneProxy,
 				{
 					// Only access the irradiance maps on the RT, even though they belong to the USkyLightComponent, 
 					// Because FScene::UpdateSkyCaptureContents does not block the RT so the writes could still be in flight
-					LightSceneProxy->Initialize(InBlendFraction, InIrradianceEnvironmentMap, BlendDestinationIrradianceEnvironmentMap);
+					LightSceneProxy->Initialize(InBlendFraction, InIrradianceEnvironmentMap, BlendDestinationIrradianceEnvironmentMap, InAverageBrightness, BlendDestinationAverageBrightness);
 				});
 			}
 		}
diff --git a/Engine/Source/Runtime/Engine/Private/Materials/HLSLMaterialTranslator.h b/Engine/Source/Runtime/Engine/Private/Materials/HLSLMaterialTranslator.h
--- a/Engine/Source/Runtime/Engine/Private/Materials/HLSLMaterialTranslator.h
+++ b/Engine/Source/Runtime/Engine/Private/Materials/HLSLMaterialTranslator.h
@@ -11,6 +11,7 @@
 #include "Materials/MaterialExpressionWorldPosition.h"
 #include "Materials/MaterialExpressionTextureSample.h"
 #include "Materials/MaterialExpressionSceneTexture.h"
+#include "Materials/MaterialExpressionNoise.h"
 #include "Materials/MaterialExpressionFunctionInput.h"
 #include "Materials/MaterialExpressionFunctionOutput.h"
 #include "Materials/MaterialExpressionMaterialFunctionCall.h"
@@ -3044,8 +3045,7 @@ protected:
 				|| SceneTextureId == PPI_SceneColor);
 		}
 
-		MaterialCompilationOutput.bNeedsGBuffer = MaterialCompilationOutput.bNeedsGBuffer
-			|| SceneTextureId == PPI_DiffuseColor 
+		const bool bNeedsGBuffer = SceneTextureId == PPI_DiffuseColor 
 			|| SceneTextureId == PPI_SpecularColor
 			|| SceneTextureId == PPI_SubsurfaceColor
 			|| SceneTextureId == PPI_BaseColor
@@ -3060,6 +3060,13 @@ protected:
 			|| SceneTextureId == PPI_StoredBaseColor
 			|| SceneTextureId == PPI_StoredSpecular;
 
+		MaterialCompilationOutput.bNeedsGBuffer = MaterialCompilationOutput.bNeedsGBuffer || bNeedsGBuffer;
+
+		if (bNeedsGBuffer && IsForwardShadingEnabled(FeatureLevel))
+		{
+			Errorf(TEXT("GBuffer scene textures not available with forward shading."));
+		}
+
 		// not yet tracked:
 		//   PPI_SeparateTranslucency, PPI_CustomDepth, PPI_AmbientOcclusion
 	}
@@ -4137,7 +4144,17 @@ protected:
 
 	virtual int32 Noise(int32 Position, float Scale, int32 Quality, uint8 NoiseFunction, bool bTurbulence, int32 Levels, float OutputMin, float OutputMax, float LevelScale, int32 FilterWidth, bool bTiling, uint32 RepeatSize) override
 	{
-		if (ErrorUnlessFeatureLevelSupported(ERHIFeatureLevel::SM4) == INDEX_NONE)
+		// GradientTex3D uses 3D texturing, which is not available on ES2
+		if (NoiseFunction == NOISEFUNCTION_GradientTex3D)
+		{
+			if (ErrorUnlessFeatureLevelSupported(ERHIFeatureLevel::SM4) == INDEX_NONE)
+			{
+				Errorf(TEXT("3D textures are not supported for ES2"));
+				return INDEX_NONE;
+			}
+		}
+		// all others are fine for ES2 feature level
+		else if (ErrorUnlessFeatureLevelSupported(ERHIFeatureLevel::ES2) == INDEX_NONE)
 		{
 			return INDEX_NONE;
 		}
diff --git a/Engine/Source/Runtime/Engine/Private/Materials/Material.cpp b/Engine/Source/Runtime/Engine/Private/Materials/Material.cpp
--- a/Engine/Source/Runtime/Engine/Private/Materials/Material.cpp
+++ b/Engine/Source/Runtime/Engine/Private/Materials/Material.cpp
@@ -2754,6 +2754,7 @@ bool UMaterial::CanEditChange(const UProperty* InProperty) const
 		}
 
 		if (PropertyName == GET_MEMBER_NAME_STRING_CHECKED(UMaterial, bFullyRough) ||
+			PropertyName == GET_MEMBER_NAME_STRING_CHECKED(UMaterial, bNormalCurvatureToRoughness) ||
 			PropertyName == GET_MEMBER_NAME_STRING_CHECKED(UMaterial, TwoSided) ||
 			PropertyName == GET_MEMBER_NAME_STRING_CHECKED(UMaterial, bUseLightmapDirectionality) ||
 			PropertyName == GET_MEMBER_NAME_STRING_CHECKED(UMaterial, bUseHQForwardReflections) ||
diff --git a/Engine/Source/Runtime/Engine/Private/Materials/MaterialShader.cpp b/Engine/Source/Runtime/Engine/Private/Materials/MaterialShader.cpp
--- a/Engine/Source/Runtime/Engine/Private/Materials/MaterialShader.cpp
+++ b/Engine/Source/Runtime/Engine/Private/Materials/MaterialShader.cpp
@@ -971,6 +971,19 @@ void FMaterialShaderMap::SaveForRemoteRecompile(FArchive& Ar, const TMap<FString
 				TMap<FShaderId, FShader*> ShaderList;
 				ShaderMap->GetShaderList(ShaderList);
 
+				// get shaders from shader pipelines
+				TArray<FShaderPipeline*> ShaderPipelineList;
+				ShaderMap->GetShaderPipelineList(ShaderPipelineList);
+
+				for (FShaderPipeline* ShaderPipeline : ShaderPipelineList)
+				{
+					for (FShader* Shader : ShaderPipeline->GetShaders())
+					{
+						FShaderId ShaderId = Shader->GetId();
+						ShaderList.Add(ShaderId, Shader);
+					}
+				}
+
 				// get the resources from the shaders
 				for (auto& KeyValue : ShaderList)
 				{
@@ -1126,6 +1139,7 @@ void FMaterialShaderMap::LoadForRemoteRecompile(FArchive& Ar, EShaderPlatform Sh
 							FMaterialResource* MaterialResource = MatchingMaterial->GetMaterialResource(GetMaxSupportedFeatureLevel(ShaderPlatform), (EMaterialQualityLevel::Type)QualityLevelIndex);
 
 							MaterialResource->SetGameThreadShaderMap(LoadedShaderMap);
+							MaterialResource->RegisterInlineShaderMap();
 
 							ENQUEUE_UNIQUE_RENDER_COMMAND_TWOPARAMETER(
 								FSetShaderMapOnMaterialResources,
diff --git a/Engine/Source/Runtime/Engine/Private/Materials/MaterialShared.cpp b/Engine/Source/Runtime/Engine/Private/Materials/MaterialShared.cpp
--- a/Engine/Source/Runtime/Engine/Private/Materials/MaterialShared.cpp
+++ b/Engine/Source/Runtime/Engine/Private/Materials/MaterialShared.cpp
@@ -921,6 +921,11 @@ bool FMaterialResource::IsFullyRough() const
 	return Material->bFullyRough;
 }
 
+bool FMaterialResource::UseNormalCurvatureToRoughness() const
+{
+	return Material->bNormalCurvatureToRoughness;
+}
+
 bool FMaterialResource::IsUsingFullPrecision() const
 {
 	return Material->bUseFullPrecision;
@@ -1442,6 +1447,7 @@ void FMaterial::SetupMaterialEnvironment(
 	OutEnvironment.SetDefine(TEXT("MATERIAL_SSR"), ShouldDoSSR());
 	OutEnvironment.SetDefine(TEXT("MATERIAL_BLOCK_GI"), ShouldBlockGI());
 	OutEnvironment.SetDefine(TEXT("MATERIAL_DITHER_OPACITY_MASK"), IsDitherMasked());
+	OutEnvironment.SetDefine(TEXT("MATERIAL_NORMAL_CURVATURE_TO_ROUGHNESS"), UseNormalCurvatureToRoughness() ? TEXT("1") : TEXT("0"));
 	OutEnvironment.SetDefine(TEXT("MATERIAL_ALLOW_NEGATIVE_EMISSIVECOLOR"), AllowNegativeEmissiveColor());
 
 	if (IsUsingFullPrecision())
diff --git a/Engine/Source/Runtime/Engine/Private/SceneView.cpp b/Engine/Source/Runtime/Engine/Private/SceneView.cpp
--- a/Engine/Source/Runtime/Engine/Private/SceneView.cpp
+++ b/Engine/Source/Runtime/Engine/Private/SceneView.cpp
@@ -230,7 +230,8 @@ static TAutoConsoleVariable<int32> CVarDefaultAntiAliasing(
 	TEXT("Engine default (project setting) for AntiAliasingMethod is (postprocess volume/camera/game setting still can override)\n")
 	TEXT(" 0: off (no anti-aliasing)\n")
 	TEXT(" 1: FXAA (faster than TemporalAA but much more shimmering for non static cases)\n")
-	TEXT(" 2: TemporalAA (default)"));
+	TEXT(" 2: TemporalAA (default)\n")
+	TEXT(" 3: MSAA (Forward shading only)"));
 
 static TAutoConsoleVariable<float> CVarMotionBlurScale(
 	TEXT("r.MotionBlur.Scale"),
@@ -457,6 +458,7 @@ FSceneView::FSceneView(const FSceneViewInitOptions& InitOptions)
 	, bAllowTranslucentPrimitivesInHitProxy( true )
 	, bHasSelectedComponents( false )
 #endif
+	, AntiAliasingMethod(AAM_None)
 	, FeatureLevel(InitOptions.ViewFamily ? InitOptions.ViewFamily->GetFeatureLevel() : GMaxRHIFeatureLevel)
 {
 	check(UnscaledViewRect.Min.X >= 0);
@@ -665,6 +667,66 @@ FSceneView::FSceneView(const FSceneViewInitOptions& InitOptions)
 
 	static const auto MultiViewCVar = IConsoleManager::Get().FindTConsoleVariableDataInt(TEXT("vr.MultiView"));
 	bIsMultiViewEnabled = ShaderPlatform == EShaderPlatform::SP_PS4 && (MultiViewCVar && MultiViewCVar->GetValueOnAnyThread() != 0);
+
+	SetupAntiAliasingMethod();
+}
+
+void FSceneView::SetupAntiAliasingMethod()
+{
+	{
+		int32 Value = CVarDefaultAntiAliasing.GetValueOnAnyThread();
+		if (Value >= 0 && Value < AAM_MAX)
+		{
+			AntiAliasingMethod = (EAntiAliasingMethod)Value;
+		}
+	}
+
+	static const auto CVarMobileMSAA = IConsoleManager::Get().FindTConsoleVariableDataInt(TEXT("r.MobileMSAA"));
+	if (FeatureLevel <= ERHIFeatureLevel::ES3_1 && CVarMobileMSAA ? CVarMobileMSAA->GetValueOnAnyThread() > 1 : false)
+	{
+		//@todo Ronin check what we support under MSAA
+
+		// Turn off various features which won't work with mobile MSAA.
+		//FinalPostProcessSettings.DepthOfFieldScale = 0.0f;
+		AntiAliasingMethod = AAM_None;
+	}
+
+	if (Family)
+	{
+		static IConsoleVariable* CVarMSAACount = IConsoleManager::Get().FindConsoleVariable(TEXT("r.MSAACount"));
+
+		if (AntiAliasingMethod == AAM_MSAA && IsForwardShadingEnabled(FeatureLevel) && CVarMSAACount->GetInt() <= 1)
+		{
+			// Fallback to temporal AA so we can easily toggle methods with r.MSAACount
+			AntiAliasingMethod = AAM_TemporalAA;
+		}
+
+		static const auto PostProcessAAQualityCVar = IConsoleManager::Get().FindTConsoleVariableDataInt(TEXT("r.PostProcessAAQuality"));
+		static auto* MobileHDRCvar = IConsoleManager::Get().FindTConsoleVariableDataInt(TEXT("r.MobileHDR"));
+		static auto* MobileMSAACvar = IConsoleManager::Get().FindTConsoleVariableDataInt(TEXT("r.MobileMSAA"));
+		static uint32 MSAAValue = GShaderPlatformForFeatureLevel[FeatureLevel] == SP_OPENGL_ES2_IOS ? 1 : MobileMSAACvar->GetValueOnAnyThread();
+
+		int32 Quality = FMath::Clamp(PostProcessAAQualityCVar->GetValueOnAnyThread(), 0, 6);
+		const bool bWillApplyTemporalAA = Family->EngineShowFlags.PostProcessing || bIsPlanarReflection;
+
+		if (!bWillApplyTemporalAA || !Family->EngineShowFlags.AntiAliasing || Quality <= 0
+			// Disable antialiasing in GammaLDR mode to avoid jittering.
+			|| (FeatureLevel == ERHIFeatureLevel::ES2 && MobileHDRCvar->GetValueOnAnyThread() == 0)
+			|| (FeatureLevel <= ERHIFeatureLevel::ES3_1 && (MSAAValue > 1))
+			|| Family->EngineShowFlags.VisualizeBloom
+			|| Family->EngineShowFlags.VisualizeDOF)
+		{
+			AntiAliasingMethod = AAM_None;
+		}
+
+		if (AntiAliasingMethod == AAM_TemporalAA)
+		{
+			if (!Family->EngineShowFlags.TemporalAA || !Family->bRealtimeUpdate || Quality < 3)
+			{
+				AntiAliasingMethod = AAM_FXAA;
+			}
+		}
+	}
 }
 
 static TAutoConsoleVariable<int32> CVarCompensateForFOV(
@@ -1217,11 +1279,6 @@ void FSceneView::OverridePostProcessSettings(const FPostProcessSettings& Src, fl
 		{
 			Dest.AmbientOcclusionRadiusInWS = Src.AmbientOcclusionRadiusInWS;
 		}
-
-		if (Src.bOverride_AntiAliasingMethod)
-		{
-			Dest.AntiAliasingMethod = Src.AntiAliasingMethod;
-		}
 	}
 	
 	// will be deprecated soon, use the new asset LightPropagationVolumeBlendable instead
@@ -1373,14 +1430,6 @@ void FSceneView::StartFinalPostprocessSettings(FVector InViewLocation)
 			FinalPostProcessSettings.LensFlareIntensity = 0;
 		}
 
-		{
-			int32 Value = CVarDefaultAntiAliasing.GetValueOnGameThread();
-			if (Value >= 0 && Value < AAM_MAX)
-			{
-				FinalPostProcessSettings.AntiAliasingMethod = (EAntiAliasingMethod)Value;
-			}
-		}
-
 		{
 			int32 Value = CVarDefaultAmbientOcclusionStaticFraction.GetValueOnGameThread();
 
@@ -1446,18 +1495,6 @@ void FSceneView::EndFinalPostprocessSettings(const FSceneViewInitOptions& ViewIn
 		}
 	}
 
-	{
-		static const auto CVarMobileMSAA = IConsoleManager::Get().FindTConsoleVariableDataInt(TEXT("r.MobileMSAA"));
-		if (FeatureLevel <= ERHIFeatureLevel::ES3_1 && CVarMobileMSAA ? CVarMobileMSAA->GetValueOnGameThread() > 1 : false)
-		{
-			//@todo Ronin check what we support under MSAA
-
-			// Turn off various features which won't work with mobile MSAA.
-			//FinalPostProcessSettings.DepthOfFieldScale = 0.0f;
-			FinalPostProcessSettings.AntiAliasingMethod = AAM_None;
-		}
-	}
-
 	{
 		static const auto SceneColorFringeQualityCVar = IConsoleManager::Get().FindTConsoleVariableDataInt(TEXT("r.SceneColorFringeQuality"));
 
@@ -1709,36 +1746,6 @@ void FSceneView::EndFinalPostprocessSettings(const FSceneViewInitOptions& ViewIn
 		FinalPostProcessSettings.IndirectLightingIntensity = 0.0f;
 	}
 
-	// Anti-Aliasing
-	{
-		static const auto PostProcessAAQualityCVar = IConsoleManager::Get().FindTConsoleVariableDataInt(TEXT("r.PostProcessAAQuality")); 
-		static auto* MobileHDRCvar = IConsoleManager::Get().FindTConsoleVariableDataInt(TEXT("r.MobileHDR"));
-		static auto* MobileMSAACvar = IConsoleManager::Get().FindTConsoleVariableDataInt(TEXT("r.MobileMSAA"));
-		static uint32 MSAAValue = GShaderPlatformForFeatureLevel[SceneViewFeatureLevel] == SP_OPENGL_ES2_IOS ? 1 : MobileMSAACvar->GetValueOnGameThread();
-
-		int32 Quality = FMath::Clamp(PostProcessAAQualityCVar->GetValueOnGameThread(), 0, 6);
-		const bool bWillApplyTemporalAA = Family->EngineShowFlags.PostProcessing || bIsPlanarReflection;
-
-		if( !bWillApplyTemporalAA || !Family->EngineShowFlags.AntiAliasing || Quality <= 0
-			// Disable antialiasing in GammaLDR mode to avoid jittering.
-			|| (SceneViewFeatureLevel == ERHIFeatureLevel::ES2 && MobileHDRCvar->GetValueOnGameThread() == 0)
-			|| (SceneViewFeatureLevel <= ERHIFeatureLevel::ES3_1 && (MSAAValue > 1))
-			|| Family->EngineShowFlags.VisualizeBloom
-			|| Family->EngineShowFlags.VisualizeDOF)
-		{
-			FinalPostProcessSettings.AntiAliasingMethod = AAM_None;
-		}
-
-		if( FinalPostProcessSettings.AntiAliasingMethod == AAM_TemporalAA)
-		{
-			if( !Family->EngineShowFlags.TemporalAA || !Family->bRealtimeUpdate || Quality < 3 )
-			{
-				FinalPostProcessSettings.AntiAliasingMethod = AAM_FXAA;
-			}
-		}
-
-	}
-
 	if (AllowDebugViewmodes())
 	{
 		ConfigureBufferVisualizationSettings();
diff --git a/Engine/Source/Runtime/Engine/Private/ShaderCompiler/ShaderCompiler.cpp b/Engine/Source/Runtime/Engine/Private/ShaderCompiler/ShaderCompiler.cpp
--- a/Engine/Source/Runtime/Engine/Private/ShaderCompiler/ShaderCompiler.cpp
+++ b/Engine/Source/Runtime/Engine/Private/ShaderCompiler/ShaderCompiler.cpp
@@ -2565,6 +2565,11 @@ void GlobalBeginCompileShader(
 		Input.Environment.SetDefine(TEXT("FORWARD_SHADING"), CVar ? (CVar->GetInt() != 0) : 0);
 	}
 
+	{
+		static IConsoleVariable* CVar = IConsoleManager::Get().FindConsoleVariable(TEXT("r.VertexFoggingForOpaque"));
+		Input.Environment.SetDefine(TEXT("VERTEX_FOGGING_FOR_OPAQUE"), CVar ? (CVar->GetInt() != 0) : 0);
+	}
+
 	if (GSupportsRenderTargetWriteMask)
 	{
 		Input.Environment.SetDefine(TEXT("PLATFORM_SUPPORTS_RENDERTARGET_WRITE_MASK"), 1);
@@ -2719,6 +2724,12 @@ public:
 			// fixup uniform expressions
 			UMaterialInterface::RecacheAllMaterialUniformExpressions();
 		}
+
+		ENQUEUE_UNIQUE_RENDER_COMMAND(
+			FRecreateBoundShaderStates,
+		{
+			RHIRecreateRecursiveBoundShaderStates();
+		});
 	}
 
 private:
