BC6H - remove Float16 values, which aren't correctly handled by the ISPCTextureCompressor. Contributed by [REDACTED_PERSON].

diff --git a/Engine/Source/Developer/TextureFormatIntelISPCTexComp/Private/TextureFormatIntelISPCTexComp.cpp b/Engine/Source/Developer/TextureFormatIntelISPCTexComp/Private/TextureFormatIntelISPCTexComp.cpp
--- a/Engine/Source/Developer/TextureFormatIntelISPCTexComp/Private/TextureFormatIntelISPCTexComp.cpp
+++ b/Engine/Source/Developer/TextureFormatIntelISPCTexComp/Private/TextureFormatIntelISPCTexComp.cpp
@@ -645,6 +645,31 @@ public:
 		InOutImage.SizeY = AlignedSizeY;
 	}
 
+	/**
+	Remove Float16 values, which aren't correctly handled by the ISPCTextureCompressor.
+	https://docs.microsoft.com/en-us/windows/desktop/direct3d11/bc6h-format
+	*/
+	static void SanitizeFloat16ForBC6H(FImage& InOutImage)
+	{
+		check(InOutImage.Format == ERawImageFormat::RGBA16F);
+
+		const int32 TexelNum = InOutImage.RawData.Num() / sizeof(FFloat16);
+		FFloat16* Data = reinterpret_cast<FFloat16*>(&InOutImage.RawData[0]);
+		for (int TexelIndex = 0; TexelIndex < TexelNum; ++TexelIndex)
+		{
+			FFloat16& F16Value = Data[TexelIndex];
+
+			const bool bDenormalOrZero = F16Value.Components.Exponent == 0;
+			const bool bNegative = F16Value.Components.Sign == 1;
+
+			// Flush denormals and negative values to 0.
+			if (bDenormalOrZero || bNegative)
+			{
+				F16Value = 0;
+			}
+		}
+	}
+
 	virtual bool CompressImage(
 		const FImage& InImage,
 		const struct FTextureBuildSettings& BuildSettings,
@@ -670,6 +695,8 @@ public:
 			FImage Image;
 			InImage.CopyTo(Image, ERawImageFormat::RGBA16F, EGammaSpace::Linear);
 
+			SanitizeFloat16ForBC6H(Image);
+
 			bc6h_enc_settings settings;
 			GetProfile_bc6h_basic(&settings);
 
