BC6H - don't need to flush input denormals to 0, it's enough just to clamp negative values.

diff --git a/Engine/Source/Developer/TextureFormatIntelISPCTexComp/Private/TextureFormatIntelISPCTexComp.cpp b/Engine/Source/Developer/TextureFormatIntelISPCTexComp/Private/TextureFormatIntelISPCTexComp.cpp
--- a/Engine/Source/Developer/TextureFormatIntelISPCTexComp/Private/TextureFormatIntelISPCTexComp.cpp
+++ b/Engine/Source/Developer/TextureFormatIntelISPCTexComp/Private/TextureFormatIntelISPCTexComp.cpp
@@ -658,13 +658,10 @@ public:
 		FFloat16* Data = reinterpret_cast<FFloat16*>(&InOutImage.RawData[0]);
 		for (int TexelIndex = 0; TexelIndex < TexelNum; ++TexelIndex)
 		{
+			// Flush negative values to 0, as those aren't supported by BC6H_UF16.
 			FFloat16& F16Value = Data[TexelIndex];
-
-			const bool bDenormalOrZero = F16Value.Components.Exponent == 0;
 			const bool bNegative = F16Value.Components.Sign == 1;
-
-			// Flush denormals and negative values to 0.
-			if (bDenormalOrZero || bNegative)
+			if (bNegative)
 			{
 				F16Value = 0;
 			}
