UE4.25 - Fix for ISPCTexComp dll not loading on projects on a different directory

diff --git a/Engine/Source/Developer/TextureFormatIntelISPCTexComp/Private/TextureFormatIntelISPCTexComp.cpp b/Engine/Source/Developer/TextureFormatIntelISPCTexComp/Private/TextureFormatIntelISPCTexComp.cpp
--- a/Engine/Source/Developer/TextureFormatIntelISPCTexComp/Private/TextureFormatIntelISPCTexComp.cpp
+++ b/Engine/Source/Developer/TextureFormatIntelISPCTexComp/Private/TextureFormatIntelISPCTexComp.cpp
@@ -839,7 +839,7 @@ public:
 /**
  * Module for DXT texture compression.
  */
-static ITextureFormat* Singleton = NULL;
+static ITextureFormat* Singleton = nullptr;
 
 class FTextureFormatIntelISPCTexCompModule : public ITextureFormatModule
 {
@@ -852,7 +852,7 @@ public:
 	virtual ~FTextureFormatIntelISPCTexCompModule()
 	{
 		delete Singleton;
-		Singleton = NULL;
+		Singleton = nullptr;
 
 		if ( mDllHandle != nullptr )
 		{
@@ -865,17 +865,29 @@ public:
 	{
 		if (!Singleton)
 		{
+			FString DLLPath;
 #if PLATFORM_WINDOWS
 	#if PLATFORM_64BITS
-			mDllHandle = FPlatformProcess::GetDllHandle(TEXT("../../../Engine/Binaries/ThirdParty/IntelISPCTexComp/Win64-Release/ispc_texcomp.dll"));
+			DLLPath = FPaths::EngineDir() / TEXT("Binaries/ThirdParty/IntelISPCTexComp/Win64-Release/ispc_texcomp.dll");
 	#else	//32-bit platform
-			mDllHandle = FPlatformProcess::GetDllHandle(TEXT("../../../Engine/Binaries/ThirdParty/IntelISPCTexComp/Win32-Release/ispc_texcomp.dll"));
+			DLLPath = FPaths::EngineDir() / TEXT("Binaries/ThirdParty/IntelISPCTexComp/Win32-Release/ispc_texcomp.dll");
 	#endif
 #elif PLATFORM_MAC
-			mDllHandle = FPlatformProcess::GetDllHandle(TEXT("libispc_texcomp.dylib"));
+			DLLPath = TEXT("libispc_texcomp.dylib");
 #elif PLATFORM_LINUX
-			mDllHandle = FPlatformProcess::GetDllHandle(TEXT("../../../Engine/Binaries/ThirdParty/IntelISPCTexComp/Linux64-Release/libispc_texcomp.so"));
+			DLLPath = FPaths::EngineDir() / TEXT("ThirdParty/IntelISPCTexComp/Linux64-Release/libispc_texcomp.so");
 #endif
+
+			if (DLLPath.Len() > 0)
+			{
+				mDllHandle = FPlatformProcess::GetDllHandle(*DLLPath);
+				UE_CLOG(mDllHandle == nullptr, LogTextureFormatIntelISPCTexComp, Warning, TEXT("Unable to load %s"), *DLLPath);
+			}
+			else
+			{
+				UE_LOG(LogTextureFormatIntelISPCTexComp, Warning, TEXT("Platform does not have an ispc_texcomp DLL/library"));
+			}
+
 			Singleton = new FTextureFormatIntelISPCTexComp();
 		}
 		return Singleton;
