UE4.25 - Linux fix for ISPCTexComp dll not loading on projects on a different directory

Fixes this warning on UE4Editor startup:

  LogCore: Warning: dlopen failed: /epic/UE4.git/Engine/ThirdParty/IntelISPCTexComp/Linux64-Release/libispc_texcomp.so: cannot open shared object file: No such file or directory
  LogTextureFormatIntelISPCTexComp: Warning: Unable to load ../../../Engine/ThirdParty/IntelISPCTexComp/Linux64-Release/libispc_texcomp.so

diff --git a/Engine/Source/Developer/TextureFormatIntelISPCTexComp/Private/TextureFormatIntelISPCTexComp.cpp b/Engine/Source/Developer/TextureFormatIntelISPCTexComp/Private/TextureFormatIntelISPCTexComp.cpp
--- a/Engine/Source/Developer/TextureFormatIntelISPCTexComp/Private/TextureFormatIntelISPCTexComp.cpp
+++ b/Engine/Source/Developer/TextureFormatIntelISPCTexComp/Private/TextureFormatIntelISPCTexComp.cpp
@@ -875,7 +875,7 @@ public:
 #elif PLATFORM_MAC
 			DLLPath = TEXT("libispc_texcomp.dylib");
 #elif PLATFORM_LINUX
-			DLLPath = FPaths::EngineDir() / TEXT("ThirdParty/IntelISPCTexComp/Linux64-Release/libispc_texcomp.so");
+			DLLPath = FPaths::EngineDir() / TEXT("Binaries/ThirdParty/IntelISPCTexComp/Linux64-Release/libispc_texcomp.so");
 #endif
 
 			if (DLLPath.Len() > 0)
