nop comment about bad DDC Key behavior of TextureFormatASTC with ISPC

diff --git a/Engine/Source/Developer/TextureFormatASTC/Private/TextureFormatASTC.cpp b/Engine/Source/Developer/TextureFormatASTC/Private/TextureFormatASTC.cpp
--- a/Engine/Source/Developer/TextureFormatASTC/Private/TextureFormatASTC.cpp
+++ b/Engine/Source/Developer/TextureFormatASTC/Private/TextureFormatASTC.cpp
@@ -443,11 +443,32 @@ public:
 		const FTextureBuildSettings* BuildSettings = nullptr
 	) const override
 	{
+	/*
+	// note/bug , doesn't call through to ISPC :
+#if SUPPORTS_ISPC_ASTC
+		if(GASTCCompressor == 0)
+		{
+			return IntelISPCTexCompFormat.GetVersion(Format,BuildSettings);
+		}
+#endif
+	*/
+
 		return BASE_ASTC_FORMAT_VERSION;
 	}
 
 	virtual FString GetDerivedDataKeyString(const FTextureBuildSettings& BuildSettings) const override
 	{
+	/*
+	// note/bug : doesn't call through to ISPC !
+	// but it is putting GASTCCompressor in the key
+#if SUPPORTS_ISPC_ASTC
+		if(GASTCCompressor == 0)
+		{
+			return IntelISPCTexCompFormat.GetDerivedDataKeyString(BuildSettings);
+		}
+#endif
+	*/
+
 		return FString::Printf(TEXT("ASTCCmpr_%d_%d_%d"), GetQualityVersion(BuildSettings.FormatConfigOverride, BuildSettings.CompressionQuality), GASTCCompressor, int32(GASTCHDRProfile > 0 && BuildSettings.bHDRSource));
 	}
 
diff --git a/Engine/Source/Developer/TextureFormatIntelISPCTexComp/Private/TextureFormatIntelISPCTexComp.cpp b/Engine/Source/Developer/TextureFormatIntelISPCTexComp/Private/TextureFormatIntelISPCTexComp.cpp
--- a/Engine/Source/Developer/TextureFormatIntelISPCTexComp/Private/TextureFormatIntelISPCTexComp.cpp
+++ b/Engine/Source/Developer/TextureFormatIntelISPCTexComp/Private/TextureFormatIntelISPCTexComp.cpp
@@ -428,6 +428,7 @@ static EPixelFormat GetQualityFormat(int& BlockWidth, int& BlockHeight, const FC
 	EPixelFormat Format = PF_Unknown;
 	switch (OverrideSizeValue >= 0 ? OverrideSizeValue : GetDefaultCompressionBySizeValue(InFormatConfigOverride))
 	{
+		// NOTE: different than TextureFormatASTC which has the 10 and 12 size blocks here
 		case 0:	//Format = PF_ASTC_12x12; BlockWidth = BlockHeight = 12; break;
 		case 1:	//Format = PF_ASTC_10x10; BlockWidth = BlockHeight = 10; break;
 		case 2:	Format = PF_ASTC_8x8; BlockWidth = BlockHeight = 8; break;
@@ -633,6 +634,18 @@ public:
 	{
 		return BASE_ISPC_DX11_FORMAT_VERSION;
 	}
+	
+	/*
+	virtual FString GetDerivedDataKeyString(const FTextureBuildSettings& BuildSettings) const override
+	{
+		// @@!! no GetDerivedDataKeyString !?
+
+		// should log settings here :
+		// should store
+		// GetDefaultCompressionBySizeValue() 
+		//return FString::Printf(TEXT("ISPC_%d_%d_%d"), GetQualityVersion(BuildSettings.FormatConfigOverride, BuildSettings.CompressionQuality), GASTCCompressor, int32(GASTCHDRProfile > 0 && BuildSettings.bHDRSource));
+	}
+	*/
 
 	virtual void GetSupportedFormats(TArray<FName>& OutFormats) const override
 	{
@@ -741,6 +754,7 @@ public:
 			{
 				F16Value.Encoded = 0;
 			}
+			// @@!! also clamp +Inf to max non-inf value?
 		}
 	}
 
