ISPCTexComp ASTC fix need to use alpha profile for NormalAG
also improve debug dump faccility

diff --git a/Engine/Source/Developer/TextureFormatIntelISPCTexComp/Private/TextureFormatIntelISPCTexComp.cpp b/Engine/Source/Developer/TextureFormatIntelISPCTexComp/Private/TextureFormatIntelISPCTexComp.cpp
--- a/Engine/Source/Developer/TextureFormatIntelISPCTexComp/Private/TextureFormatIntelISPCTexComp.cpp
+++ b/Engine/Source/Developer/TextureFormatIntelISPCTexComp/Private/TextureFormatIntelISPCTexComp.cpp
@@ -19,6 +19,7 @@
 #include "TextureBuildFunction.h"
 #include "DerivedDataBuildFunctionFactory.h"
 #include "DerivedDataSharedString.h"
+#include "Misc/WildcardString.h"
 
 #include "ispc_texcomp.h"
 
@@ -89,6 +90,10 @@ class FIntelISPCTexCompTextureBuildFunction final : public FTextureBuildFunction
 #define BLOCK_SIZE_IN_BYTES 16
 
 
+#if DEBUG_SAVE_INTERMEDIATE_IMAGES
+namespace
+{
+
 	// Bitmap compression types.
 enum EBitmapCompression
 {
@@ -143,7 +148,7 @@ struct FBitmapInfoHeader
 #pragma pack(pop)
 
 
-void SaveImageAsBMP( FArchive& Ar, const uint8* RawData, int SourceBytesPerPixel, int SizeX, int SizeY )
+static void SaveImageAsBMP( FArchive& Ar, const uint8* RawData, int SourceBytesPerPixel, int SizeX, int SizeY )
 {
 	FBitmapFileHeader bmf;
 	FBitmapInfoHeader bmh;
@@ -226,7 +231,7 @@ struct astc_header
 };
 #pragma pack(pop)
 
-void SaveImageAsASTC(FArchive& Ar, uint8* RawData, int SizeX, int SizeY, int block_width, int block_height)
+static void SaveImageAsASTC(FArchive& Ar, uint8* RawData, int SizeX, int SizeY, int block_width, int block_height)
 {
 	astc_header file_header;
 
@@ -252,6 +257,10 @@ void SaveImageAsASTC(FArchive& Ar, uint8* RawData, int SizeX, int SizeY, int blo
 	Ar.Serialize(RawData, height_in_blocks * stride);
 }
 
+}; // namespace
+#endif 
+// DEBUG_SAVE_INTERMEDIATE_IMAGES
+
 struct FMultithreadSettings
 {
 	int iScansPerTask;
@@ -481,6 +490,13 @@ static EPixelFormat GetASTCQualityFormat(int& BlockWidth, int& BlockHeight, cons
 struct FASTCEncoderSettings : public astc_enc_settings
 {
 	FName TextureFormatName;
+
+	FASTCEncoderSettings()
+	{
+		// ensure astc_enc_settings are initialized :
+		astc_enc_settings * parent = this;
+		memset(parent,0,sizeof(*parent));
+	}
 };
 
 
@@ -584,6 +600,10 @@ static void IntelASTCCompressScans(FASTCEncoderSettings* pEncSettings, FImage* p
 			}
 		}
 	}
+	else
+	{
+		checkNoEntry();
+	}
 
 	rgba_surface insurface;
 	insurface.ptr		= pInTexels + (yStart * InStride);
@@ -592,7 +612,7 @@ static void IntelASTCCompressScans(FASTCEncoderSettings* pEncSettings, FImage* p
 	insurface.stride	= pInImage->SizeX * 4;
 	
 	check((yStart % pEncSettings->block_height) == 0);
-	pOutTexels += yStart / pEncSettings->block_height * OutStride;
+	pOutTexels += (yStart / pEncSettings->block_height) * OutStride;
 	CompressBlocksASTC(&insurface, pOutTexels, pEncSettings);
 }
 
@@ -899,7 +919,20 @@ public:
 			
 			if ( bIsNormalMap )
 			{
-				GetProfile_astc_fast(&EncoderSettings, BlockWidth, BlockHeight);
+				check( BlockWidth == 6 && BlockHeight == 6 );
+
+				if (BuildSettings.TextureFormatName == GTextureFormatNameASTC_NormalAG)
+				{
+					GetProfile_astc_alpha_fast(&EncoderSettings, BlockWidth, BlockHeight);
+				}
+				else if (BuildSettings.TextureFormatName == GTextureFormatNameASTC_NormalRG)
+				{
+					GetProfile_astc_fast(&EncoderSettings, BlockWidth, BlockHeight);
+				}
+				else
+				{	
+					checkNoEntry();
+				}
 			}
 			else
 			{			
@@ -930,17 +963,44 @@ public:
 				PadImageToBlockSize(Image, EncoderSettings.block_width, EncoderSettings.block_height, 4 * 1);
 				
 #if DEBUG_SAVE_INTERMEDIATE_IMAGES
-				//@DEBUG (save padded input as BMP):
-				static bool SaveInputOutput = false;
-				static volatile int32 Counter = 0;
-				int LocalCounter = Counter;
-				if (SaveInputOutput)	// && LocalCounter < 10 && Image.SizeX >= 1024)
+				// To DebugDump specific files, modify DebugDumpFilter here
+				//  (currently not exposed to command line, but could be)
+				//static FString DebugDumpFilter(TEXT("*DummySpriteTexture*"));;
+				static FString DebugDumpFilter(TEXT("*"));
+				//static FString DebugDumpFilter;
+
+				bool SaveInputOutput = false;
+				if ( ! DebugDumpFilter.IsEmpty() )
+				{
+					SaveInputOutput = FWildcardString::IsMatchSubstring(*DebugDumpFilter, DebugTexturePathName.GetData(), DebugTexturePathName.GetData() + DebugTexturePathName.Len(), ESearchCase::IgnoreCase);
+				}
+
+				FString DebugDumpFileNameBase;
+				if (SaveInputOutput)
+				{
+					FString FileName = FString::Printf(TEXT("%.*s_%dx%d_%s_%dx%d"), DebugTexturePathName.Len(), DebugTexturePathName.GetData(), 
+						Image.SizeX, Image.SizeY, *BuildSettings.TextureFormatName.ToString(), BlockWidth,BlockHeight);
+
+					// Object paths a) can contain slashes as its a path, and we dont want a hierarchy and b) can have random characters we don't want
+					FileName = FPaths::MakeValidFileName(FileName, TEXT('_'));
+
+					FileName = FPaths::ProjectSavedDir() + TEXT("ISPCDebugImages/") + FileName;
+					
+					// limit to _MAX_PATH
+					if ( FileName.Len() >= 256 )
+					{
+						FileName = FileName.Right(255);
+					}
+
+					DebugDumpFileNameBase = FileName;
+				}
+
+				if (SaveInputOutput)
 				{
-					const FString FileName = FString::Printf(TEXT("Smedis-Input-%d.bmp"), FPlatformTLS::GetCurrentThreadId());
+					const FString FileName = FString::Printf(TEXT("%s-in.bmp"), *DebugDumpFileNameBase );
 					FArchive* FileWriter = IFileManager::Get().CreateFileWriter(*FileName);
 					SaveImageAsBMP(*FileWriter, Image.RawData.GetData(), 4, Image.SizeX, Image.SizeY);
 					delete FileWriter;
-					FPlatformAtomics::InterlockedIncrement(&Counter);
 				}
 #endif
 
@@ -949,20 +1009,19 @@ public:
 				bCompressionSucceeded = true;
 				
 #if DEBUG_SAVE_INTERMEDIATE_IMAGES
-				//@DEBUG (save swizzled/fixed-up input as BMP):
-				if (SaveInputOutput)	// && LocalCounter < 10 && Image.SizeX >= 1024)
+				// (save swizzled/fixed-up input as BMP):
+				if (SaveInputOutput)
 				{
-					const FString FileName = FString::Printf(TEXT("Smedis-InputSwizzled-%d.bmp"), FPlatformTLS::GetCurrentThreadId());
+					const FString FileName = FString::Printf(TEXT("%s-in-swiz.bmp"), *DebugDumpFileNameBase );
 					FArchive* FileWriter = IFileManager::Get().CreateFileWriter(*FileName);
 					SaveImageAsBMP(*FileWriter, Image.RawData.GetData(), 4, Image.SizeX, Image.SizeY);
 					delete FileWriter;
-					FPlatformAtomics::InterlockedIncrement(&Counter);
 				}
 
-				//@DEBUG (save output as .astc file):
-				if (SaveInputOutput)// && LocalCounter < 10 && Image.SizeX >= 1024)
+				// (save output as .astc file):
+				if (SaveInputOutput)
 				{
-					const FString FileName = FString::Printf(TEXT("Smedis-Output-%d.astc"), FPlatformTLS::GetCurrentThreadId());
+					const FString FileName = FString::Printf(TEXT("%s-out.astc"), *DebugDumpFileNameBase );
 					FArchive* FileWriter = IFileManager::Get().CreateFileWriter(*FileName);
 					SaveImageAsASTC(*FileWriter, OutCompressedImage.RawData.GetData(), OutCompressedImage.SizeX, OutCompressedImage.SizeY, EncoderSettings.block_width, EncoderSettings.block_height);
 					delete FileWriter;
