Add full path to macOS dylib

diff --git a/Engine/Source/Developer/TextureFormatIntelISPCTexComp/Private/TextureFormatIntelISPCTexComp.cpp b/Engine/Source/Developer/TextureFormatIntelISPCTexComp/Private/TextureFormatIntelISPCTexComp.cpp
--- a/Engine/Source/Developer/TextureFormatIntelISPCTexComp/Private/TextureFormatIntelISPCTexComp.cpp
+++ b/Engine/Source/Developer/TextureFormatIntelISPCTexComp/Private/TextureFormatIntelISPCTexComp.cpp
@@ -647,7 +647,7 @@ public:
 #if PLATFORM_WINDOWS
 		DLLPath = FPaths::EngineDir() / TEXT("Binaries/ThirdParty/Intel/ISPCTexComp/Win64-Release/ispc_texcomp.dll");
 #elif PLATFORM_MAC
-		DLLPath = TEXT("libispc_texcomp.dylib");
+		DLLPath = FPaths::EngineDir() / TEXT("Binaries/ThirdParty/Intel/ISPCTexComp/Mac64-Release/libispc_texcomp.dylib");
 #elif PLATFORM_LINUX
 		DLLPath = FPaths::EngineDir() / TEXT("Binaries/ThirdParty/Intel/ISPCTexComp/Linux64-Release/libispc_texcomp.so");
 #endif
