ISPCTexcomp fix inconsistent behavior normalizing normals unlike all other TextureFormats
maintain legacy behavior for existing textures

diff --git a/Engine/Source/Developer/TextureFormatIntelISPCTexComp/Private/TextureFormatIntelISPCTexComp.cpp b/Engine/Source/Developer/TextureFormatIntelISPCTexComp/Private/TextureFormatIntelISPCTexComp.cpp
--- a/Engine/Source/Developer/TextureFormatIntelISPCTexComp/Private/TextureFormatIntelISPCTexComp.cpp
+++ b/Engine/Source/Developer/TextureFormatIntelISPCTexComp/Private/TextureFormatIntelISPCTexComp.cpp
@@ -42,7 +42,7 @@ class FIntelISPCTexCompTextureBuildFunction final : public FTextureBuildFunction
 };
 
 // increment this if you change anything that will affect compression in this file
-#define BASE_ISPC_DX11_FORMAT_VERSION 8
+#define BASE_ISPC_DX11_FORMAT_VERSION 9
 
 // For debugging intermediate image results by saving them out as files.
 #define DEBUG_SAVE_INTERMEDIATE_IMAGES 0
@@ -490,6 +490,7 @@ static EPixelFormat GetASTCQualityFormat(int& BlockWidth, int& BlockHeight, cons
 struct FASTCEncoderSettings : public astc_enc_settings
 {
 	FName TextureFormatName;
+	bool bISPCTexcompNormalizeNormals = false;
 
 	FASTCEncoderSettings()
 	{
@@ -560,47 +561,87 @@ static void IntelASTCCompressScans(FASTCEncoderSettings* pEncSettings, FImage* p
 	}
 	else if (pEncSettings->TextureFormatName == GTextureFormatNameASTC_NormalAG)
 	{
-		// Switch byte order for compressors input (BGRA -> RGBA)
-		// Re-normalize
-		// Set any unused RGB components to 0, an unused A to 255.
-		for (int y = yStart; y < yEnd; ++y)
+		// input BGRA -> output 0,G,0,A
+
+		if ( pEncSettings->bISPCTexcompNormalizeNormals )
 		{
-			uint8* pInTexelsSwap = pInTexels + (y * InStride);
-			for (int x = 0; x < pInImage->SizeX; ++x)
+			// legacy behavior:
+			// Re-normalize normals before dropping components
+			// this is right, but doing it here makes this TextureFormat inconsistent with all others
+			for (int y = yStart; y < yEnd; ++y)
 			{
-				// @@CB inconsistent : ISPCTexcomp is normalizing normals but nobody else does
-				//	 should remove this and do it at the TextureCompressorModule level so its consistent
-				FVector Normal = FVector(pInTexelsSwap[2] / 255.0f * 2.0f - 1.0f, pInTexelsSwap[1] / 255.0f * 2.0f - 1.0f, pInTexelsSwap[0] / 255.0f * 2.0f - 1.0f);
-				Normal = Normal.GetSafeNormal();
-				pInTexelsSwap[0] = 0;
-				pInTexelsSwap[1] = FMath::RoundToInt((Normal.Y * 0.5f + 0.5f) * 255.f);
-				pInTexelsSwap[2] = 0;
-				pInTexelsSwap[3] = FMath::RoundToInt((Normal.X * 0.5f + 0.5f) * 255.f);
+				uint8* pInTexelsSwap = pInTexels + (y * InStride);
+				for (int x = 0; x < pInImage->SizeX; ++x)
+				{
+					FVector Normal = FVector(pInTexelsSwap[2] / 255.0f * 2.0f - 1.0f, pInTexelsSwap[1] / 255.0f * 2.0f - 1.0f, pInTexelsSwap[0] / 255.0f * 2.0f - 1.0f);
+					Normal = Normal.GetSafeNormal();
+					pInTexelsSwap[0] = 0;
+					pInTexelsSwap[1] = FMath::RoundToInt((Normal.Y * 0.5f + 0.5f) * 255.f);
+					pInTexelsSwap[2] = 0;
+					pInTexelsSwap[3] = FMath::RoundToInt((Normal.X * 0.5f + 0.5f) * 255.f);
+
+					pInTexelsSwap += 4;
+				}
+			}
+		}
+		else
+		{
+			// new preferred path: do not normalize normals here, use bNormalizeNormals in TextureCompressor
+			for (int y = yStart; y < yEnd; ++y)
+			{
+				uint8* pInTexelsSwap = pInTexels + (y * InStride);
+				for (int x = 0; x < pInImage->SizeX; ++x)
+				{
+					pInTexelsSwap[0] = 0;
+					//pInTexelsSwap[1] = pInTexelsSwap[1]; // G
+					pInTexelsSwap[3] = pInTexelsSwap[2]; // R->A
+					pInTexelsSwap[2] = 0;
 
-				pInTexelsSwap += 4;
+					pInTexelsSwap += 4;
+				}
 			}
 		}
 	}
 	else if (pEncSettings->TextureFormatName == GTextureFormatNameASTC_NormalRG)
 	{
-		// Switch byte order for compressors input (BGRA -> RGBA)
-		// Re-normalize
-		// Set any unused RGB components to 0, an unused A to 255.
-		for (int y = yStart; y < yEnd; ++y)
+		// input BGRA -> output RG,0,255
+
+		if ( pEncSettings->bISPCTexcompNormalizeNormals )
 		{
-			uint8* pInTexelsSwap = pInTexels + (y * InStride);
-			for (int x = 0; x < pInImage->SizeX; ++x)
+			// legacy behavior:
+			// Re-normalize normals before dropping components
+			// this is right, but doing it here makes this TextureFormat inconsistent with all others
+			for (int y = yStart; y < yEnd; ++y)
 			{
-				// @@CB inconsistent : ISPCTexcomp is normalizing normals but nobody else does
-				//	 should remove this and do it at the TextureCompressorModule level so its consistent
-				FVector Normal = FVector(pInTexelsSwap[2] / 255.0f * 2.0f - 1.0f, pInTexelsSwap[1] / 255.0f * 2.0f - 1.0f, pInTexelsSwap[0] / 255.0f * 2.0f - 1.0f);
-				Normal = Normal.GetSafeNormal();
-				pInTexelsSwap[0] = FMath::RoundToInt((Normal.X * 0.5f + 0.5f) * 255.f);
-				pInTexelsSwap[1] = FMath::RoundToInt((Normal.Y * 0.5f + 0.5f) * 255.f);
-				pInTexelsSwap[2] = 0;
-				pInTexelsSwap[3] = 255;
+				uint8* pInTexelsSwap = pInTexels + (y * InStride);
+				for (int x = 0; x < pInImage->SizeX; ++x)
+				{
+					FVector Normal = FVector(pInTexelsSwap[2] / 255.0f * 2.0f - 1.0f, pInTexelsSwap[1] / 255.0f * 2.0f - 1.0f, pInTexelsSwap[0] / 255.0f * 2.0f - 1.0f);
+					Normal = Normal.GetSafeNormal();
+					pInTexelsSwap[0] = FMath::RoundToInt((Normal.X * 0.5f + 0.5f) * 255.f);
+					pInTexelsSwap[1] = FMath::RoundToInt((Normal.Y * 0.5f + 0.5f) * 255.f);
+					pInTexelsSwap[2] = 0;
+					pInTexelsSwap[3] = 255;
+
+					pInTexelsSwap += 4;
+				}
+			}
+		}
+		else
+		{
+			// new preferred path: do not normalize normals here, use bNormalizeNormals in TextureCompressor
+			for (int y = yStart; y < yEnd; ++y)
+			{
+				uint8* pInTexelsSwap = pInTexels + (y * InStride);
+				for (int x = 0; x < pInImage->SizeX; ++x)
+				{
+					pInTexelsSwap[0] = pInTexelsSwap[2]; // R
+					//pInTexelsSwap[1] = pInTexelsSwap[1]; // G
+					pInTexelsSwap[2] = 0;
+					pInTexelsSwap[3] = 255;
 
-				pInTexelsSwap += 4;
+					pInTexelsSwap += 4;
+				}
 			}
 		}
 	}
@@ -901,6 +942,20 @@ public:
 				{	
 					checkNoEntry();
 				}
+
+				if ( BuildSettings.bNormalizeNormals || BuildSettings.bUseNewMipFilter )
+				{
+					// newer texture build, do NOT normalize normals here
+					// it is done by bNormalizeNormals if wanted, at the TextureCompressor level before getting to the TextureFormat
+					EncoderSettings.bISPCTexcompNormalizeNormals = false;
+				}
+				else
+				{
+					// old texture build
+					// enable legacy behavior to maintain same output
+					// we would prefer bISPCTexcompNormalizeNormals to just always be off and use bNormalizeNormals instead
+					EncoderSettings.bISPCTexcompNormalizeNormals = true;
+				}
 			}
 			else
 			{			
