Type safety set to off to fix build issue as I missed some stuff.

diff --git a/Engine/Source/Developer/TextureBuild/TextureBuild.Build.cs b/Engine/Source/Developer/TextureBuild/TextureBuild.Build.cs
--- a/Engine/Source/Developer/TextureBuild/TextureBuild.Build.cs
+++ b/Engine/Source/Developer/TextureBuild/TextureBuild.Build.cs
@@ -6,8 +6,6 @@ public class TextureBuild : ModuleRules
 {
 	public TextureBuild(ReadOnlyTargetRules Target) : base(Target)
 	{
-		UnsafeTypeCastWarningLevel = WarningLevel.Warning;
-
 		PrivateIncludePathModuleNames.AddRange(new string[]
 		{
 			"Engine",
diff --git a/Engine/Source/Developer/TextureCompressor/TextureCompressor.Build.cs b/Engine/Source/Developer/TextureCompressor/TextureCompressor.Build.cs
--- a/Engine/Source/Developer/TextureCompressor/TextureCompressor.Build.cs
+++ b/Engine/Source/Developer/TextureCompressor/TextureCompressor.Build.cs
@@ -6,8 +6,6 @@ public class TextureCompressor : ModuleRules
 {
 	public TextureCompressor(ReadOnlyTargetRules Target) : base(Target)
 	{
-		UnsafeTypeCastWarningLevel = WarningLevel.Warning;
-
 		PrivateIncludePathModuleNames.AddRange(new string[]
 		{
 			"Engine",
diff --git a/Engine/Source/Developer/TextureFormatDXT/TextureFormatDXT.Build.cs b/Engine/Source/Developer/TextureFormatDXT/TextureFormatDXT.Build.cs
--- a/Engine/Source/Developer/TextureFormatDXT/TextureFormatDXT.Build.cs
+++ b/Engine/Source/Developer/TextureFormatDXT/TextureFormatDXT.Build.cs
@@ -6,8 +6,6 @@ public class TextureFormatDXT : ModuleRules
 {
 	public TextureFormatDXT(ReadOnlyTargetRules Target) : base(Target)
 	{
-		UnsafeTypeCastWarningLevel = WarningLevel.Warning;
-
 		PrivateIncludePathModuleNames.AddRange(new string[]
 		{
 			"DerivedDataCache",
diff --git a/Engine/Source/Developer/TextureFormatIntelISPCTexComp/TextureFormatIntelISPCTexComp.Build.cs b/Engine/Source/Developer/TextureFormatIntelISPCTexComp/TextureFormatIntelISPCTexComp.Build.cs
--- a/Engine/Source/Developer/TextureFormatIntelISPCTexComp/TextureFormatIntelISPCTexComp.Build.cs
+++ b/Engine/Source/Developer/TextureFormatIntelISPCTexComp/TextureFormatIntelISPCTexComp.Build.cs
@@ -6,8 +6,6 @@ public class TextureFormatIntelISPCTexComp : ModuleRules
 {
 	public TextureFormatIntelISPCTexComp(ReadOnlyTargetRules Target) : base(Target)
 	{
-		UnsafeTypeCastWarningLevel = WarningLevel.Warning;
-
 		PrivateIncludePathModuleNames.AddRange(new string[]
 		{
 			"DerivedDataCache",
diff --git a/Engine/Source/Editor/TextureEditor/TextureEditor.Build.cs b/Engine/Source/Editor/TextureEditor/TextureEditor.Build.cs
--- a/Engine/Source/Editor/TextureEditor/TextureEditor.Build.cs
+++ b/Engine/Source/Editor/TextureEditor/TextureEditor.Build.cs
@@ -6,8 +6,6 @@ public class TextureEditor : ModuleRules
 {
 	public TextureEditor(ReadOnlyTargetRules Target) : base(Target)
 	{
-		UnsafeTypeCastWarningLevel = WarningLevel.Warning;
-
 		DynamicallyLoadedModuleNames.AddRange(
 			new string[] {
 				"MainFrame",
